Implementation of RuntimeAssetCache

diff --git a/Engine/Config/BaseEngine.ini b/Engine/Config/BaseEngine.ini
--- a/Engine/Config/BaseEngine.ini
+++ b/Engine/Config/BaseEngine.ini
@@ -1838,6 +1838,16 @@ Foliage=(NumLODs=1,Name=LOCTEXT("FoliageLOD","Foliage"))
 [TextureTracking]
 ;TextureName=T_GD_Traffic_Crosswalk_01
 
+; Config for RuntimeAssetCache asset groups.
+; Each group can be managed separately (e.g. advertisement cache, character image cache etc.)
+; For each group a name must be specified. Size is optional and defaults to 5MB. Size in 
+; ini file is stored as number of bytes. E.g.:
+; [RuntimeAssetCache]
+; +BucketConfigs=(Name="CharacterImage", Size=3000000)
+; +BucketConfigs=(Name="Advertisement", Size=4000000)
+[RuntimeAssetCache]
++BucketConfigs=(Name="DefaultBucket", Size=10000000)
+
 ; Derived Data backend graphs
 ; Each of the below backend graph sections contains a set of nodes used to create derived data backed graph.
 ; DerivedDataBackendGraph is the default graph and the other can be specified in the command line using:
diff --git a/Engine/Source/Runtime/RuntimeAssetCache/Private/RuntimeAssetCache.cpp b/Engine/Source/Runtime/RuntimeAssetCache/Private/RuntimeAssetCache.cpp
new file mode 100644
--- /dev/null
+++ b/Engine/Source/Runtime/RuntimeAssetCache/Private/RuntimeAssetCache.cpp
@@ -0,0 +1,166 @@
+// Copyright 1998-2015 Epic Games, Inc. All Rights Reserved.
+
+#include "RuntimeAssetCachePrivatePCH.h"
+#include "RuntimeAssetCache.h"
+#include "RuntimeAssetCachePluginInterface.h"
+#include "RuntimeAssetCacheFilesystemBackend.h"
+#include "RuntimeAssetCacheBucket.h"
+#include "RuntimeAssetCacheAsyncWorker.h"
+#include "RuntimeAssetCacheEntryMetadata.h"
+
+DEFINE_STAT(STAT_RAC_ASyncWaitTime);
+
+FRuntimeAssetCache::FRuntimeAssetCache()
+{
+	FString BucketName;
+	int32 BucketSize;
+
+	TArray<FString> BucketNames;
+	GConfig->GetArray(TEXT("RuntimeAssetCache"), TEXT("BucketConfigs"), BucketNames, GEngineIni);
+
+	for (FString& Bucket : BucketNames)
+	{
+		if (FParse::Value(*Bucket, TEXT("Name="), BucketName))
+		{
+			BucketSize = FRuntimeAssetCacheBucket::DefaultBucketSize;
+			FParse::Value(*Bucket, TEXT("Size="), BucketSize);
+			FRuntimeAssetCacheBucket* BucketDesc = FRuntimeAssetCacheBackend::Get().PreLoadBucket(FName(*BucketName), BucketSize);
+			Buckets.Add(FName(*BucketName), BucketDesc);
+		}
+	}
+}
+
+FRuntimeAssetCache::~FRuntimeAssetCache()
+{
+	for (auto& Bucket : Buckets)
+	{
+		delete Bucket.Value;
+	}
+
+	Buckets.Empty();
+}
+
+int32 FRuntimeAssetCache::GetCacheSize(FName Bucket) const
+{
+	return Buckets[Bucket]->GetSize();
+}
+
+uint32 FRuntimeAssetCache::GetAsynchronous(FRuntimeAssetCacheBuilderInterface* CacheBuilder)
+{
+	int32 Handle = GetNextHandle();
+
+	/** Must return a valid handle */
+	check(Handle != 0);
+	
+	/** Make sure task isn't processed twice. */
+	check(!PendingTasks.Contains(Handle));
+
+	FAsyncTask<FRuntimeAssetCacheAsyncWorker>* AsyncTask = new FAsyncTask<FRuntimeAssetCacheAsyncWorker>(CacheBuilder, &Buckets);
+
+	{
+		FScopeLock ScopeLock(&SynchronizationObject);
+		PendingTasks.Add(Handle, AsyncTask);
+	}
+	AddToAsyncCompletionCounter(1);
+	AsyncTask->StartBackgroundTask();
+	return Handle;
+}
+
+bool FRuntimeAssetCache::GetSynchronous(FRuntimeAssetCacheBuilderInterface* CacheBuilder, TArray<uint8>& OutData)
+{
+	FAsyncTask<FRuntimeAssetCacheAsyncWorker>* AsyncTask = new FAsyncTask<FRuntimeAssetCacheAsyncWorker>(CacheBuilder, &Buckets);
+	AddToAsyncCompletionCounter(1);
+	AsyncTask->StartSynchronousTask();
+	OutData = AsyncTask->GetTask().GetData();
+	return AsyncTask->GetTask().RetrievedEntry();
+}
+
+bool FRuntimeAssetCache::ClearCache()
+{
+	/** Tell backend to clean itself up. */
+	bool bResult = FRuntimeAssetCacheBackend::Get().ClearCache();
+
+	if (!bResult)
+	{
+		return bResult;
+	}
+
+	/** If backend is cleaned up, clean up all buckets. */
+	for (auto& Bucket : Buckets)
+	{
+		Bucket.Value->Reset();
+	}
+
+	return bResult;
+}
+
+bool FRuntimeAssetCache::ClearCache(FName Bucket)
+{
+	/** Tell backend to clean bucket up. */
+	bool bResult = FRuntimeAssetCacheBackend::Get().ClearCache(Bucket);
+
+	if (!bResult)
+	{
+		return bResult;
+	}
+
+	/** If backend is cleaned up, clean up all buckets. */
+	for (auto& Bucket : Buckets)
+	{
+		Bucket.Value->Reset();
+	}
+
+	return bResult;
+}
+
+void FRuntimeAssetCache::AddToAsyncCompletionCounter(int32 Value)
+{
+	PendingTasksCounter.Add(Value);
+	check(PendingTasksCounter.GetValue() >= 0);
+}
+
+void FRuntimeAssetCache::WaitAsynchronousCompletion(uint32 Handle)
+{
+	STAT(double ThisTime = 0);
+	{
+		SCOPE_SECONDS_COUNTER(ThisTime);
+		FAsyncTask<FRuntimeAssetCacheAsyncWorker>* AsyncTask = NULL;
+		{
+			FScopeLock ScopeLock(&SynchronizationObject);
+			AsyncTask = PendingTasks.FindRef(Handle);
+		}
+		check(AsyncTask);
+		AsyncTask->EnsureCompletion();
+	}
+	INC_FLOAT_STAT_BY(STAT_RAC_ASyncWaitTime, (float)ThisTime);
+}
+
+bool FRuntimeAssetCache::GetAsynchronousResults(uint32 Handle, TArray<uint8>& OutData)
+{
+	FAsyncTask<FRuntimeAssetCacheAsyncWorker>* AsyncTask = NULL;
+	{
+		FScopeLock ScopeLock(&SynchronizationObject);
+		PendingTasks.RemoveAndCopyValue(Handle, AsyncTask);
+	}
+	check(AsyncTask);
+	if (!AsyncTask->GetTask().RetrievedEntry())
+	{
+		delete AsyncTask;
+		return false;
+	}
+	OutData = AsyncTask->GetTask().GetData();
+	delete AsyncTask;
+	check(OutData.Num());
+	return true;
+}
+
+bool FRuntimeAssetCache::PollAsynchronousCompletion(uint32 Handle)
+{
+	FAsyncTask<FRuntimeAssetCacheAsyncWorker>* AsyncTask = NULL;
+	{
+		FScopeLock ScopeLock(&SynchronizationObject);
+		AsyncTask = PendingTasks.FindRef(Handle);
+	}
+	check(AsyncTask);
+	return AsyncTask->IsDone();
+}
diff --git a/Engine/Source/Runtime/RuntimeAssetCache/Private/RuntimeAssetCache.h b/Engine/Source/Runtime/RuntimeAssetCache/Private/RuntimeAssetCache.h
new file mode 100644
--- /dev/null
+++ b/Engine/Source/Runtime/RuntimeAssetCache/Private/RuntimeAssetCache.h
@@ -0,0 +1,65 @@
+// Copyright 1998-2015 Epic Games, Inc. All Rights Reserved.
+
+#pragma once
+#include "RuntimeAssetCacheInterface.h"
+
+DECLARE_FLOAT_ACCUMULATOR_STAT_EXTERN(TEXT("RAC ASync Wait Time"), STAT_RAC_ASyncWaitTime, STATGROUP_RAC, );
+
+/** Forward declarations. */
+class FRuntimeAssetCacheBuilderInterface;
+class IRuntimeAssetCacheBackend;
+class FRuntimeAssetCacheAsyncWorker;
+class FCacheEntryMetadata;
+class FRuntimeAssetCacheBucket;
+
+/**
+ * Concrete class implementing FRuntimeAssetCacheInterface. 
+ */
+class FRuntimeAssetCache : public FRuntimeAssetCacheInterface
+{
+	/** FRuntimeAssetCacheInterface implementation */
+public:
+	virtual bool GetSynchronous(FRuntimeAssetCacheBuilderInterface* CacheBuilder, TArray<uint8>& OutData) override;
+	virtual uint32 GetAsynchronous(FRuntimeAssetCacheBuilderInterface* CacheBuilder) override;
+	virtual int32 GetCacheSize(FName Bucket) const override;
+	virtual bool ClearCache() override;
+	virtual bool ClearCache(FName Bucket) override;
+	virtual void WaitAsynchronousCompletion(uint32 Handle) override;
+	virtual bool GetAsynchronousResults(uint32 Handle, TArray<uint8>& OutData) override;
+	virtual bool PollAsynchronousCompletion(uint32 Handle) override;
+	virtual void AddToAsyncCompletionCounter(int32 Addend) override;
+	/** End of FRuntimeAssetCacheInterface implementation */
+
+public:
+	/** Constructor */
+	FRuntimeAssetCache();
+
+	/** Destructor */
+	~FRuntimeAssetCache();
+
+private:
+
+	/**
+	* Generates next handle of async worker.
+	* @return Worker handle.
+	*/
+	int32 GetNextHandle()
+	{
+		return CurrentAsyncTaskHandle.Increment();
+	}
+
+	/** Map of bucket names to their configs. */
+	TMap<FName, FRuntimeAssetCacheBucket*> Buckets;
+
+	/** Map of handle to pending task */
+	TMap<uint32, FAsyncTask<FRuntimeAssetCacheAsyncWorker>*> PendingTasks;
+
+	/** Pending tasks synchronization object. */
+	FCriticalSection SynchronizationObject;
+
+	/** Counter used to generate worker handles. */
+	FThreadSafeCounter CurrentAsyncTaskHandle;
+
+	/** Number of pending tasks. */
+	FThreadSafeCounter PendingTasksCounter;
+};
diff --git a/Engine/Source/Runtime/RuntimeAssetCache/Private/RuntimeAssetCacheAsyncWorker.cpp b/Engine/Source/Runtime/RuntimeAssetCache/Private/RuntimeAssetCacheAsyncWorker.cpp
new file mode 100644
--- /dev/null
+++ b/Engine/Source/Runtime/RuntimeAssetCache/Private/RuntimeAssetCacheAsyncWorker.cpp
@@ -0,0 +1,228 @@
+// Copyright 1998-2015 Epic Games, Inc. All Rights Reserved.
+
+#include "RuntimeAssetCachePrivatePCH.h"
+#include "RuntimeAssetCacheAsyncWorker.h"
+#include "RuntimeAssetCachePluginInterface.h"
+#include "RuntimeAssetCacheBucket.h"
+#include "RuntimeAssetCacheBackend.h"
+#include "RuntimeAssetCacheEntryMetadata.h"
+#include "RuntimeAssetCacheInterface.h"
+#include "RuntimeAssetCacheBucketScopeLock.h"
+#include "ScopeExit.h"
+#include "RuntimeAssetCacheModule.h"
+
+/** Stats */
+DEFINE_STAT(STAT_RAC_NumBuilds);
+DEFINE_STAT(STAT_RAC_NumCacheHits);
+DEFINE_STAT(STAT_RAC_NumFails);
+DEFINE_STAT(STAT_RAC_NumGets);
+DEFINE_STAT(STAT_RAC_NumPuts);
+
+FRuntimeAssetCacheAsyncWorker::FRuntimeAssetCacheAsyncWorker(FRuntimeAssetCacheBuilderInterface* InCacheBuilder
+	, TMap<FName, FRuntimeAssetCacheBucket*>* InBuckets)
+	: CacheBuilder(InCacheBuilder)
+	, Buckets(InBuckets)
+{ }
+
+void FRuntimeAssetCacheAsyncWorker::DoWork()
+{
+	const TCHAR* Bucket = CacheBuilder->GetTypeName();
+	FName BucketName = FName(Bucket);
+	FString CacheKey = BuildCacheKey(CacheBuilder);
+	FName CacheKeyName = FName(*CacheKey);
+	FRuntimeAssetCacheBucket* CurrentBucket = (*Buckets)[BucketName];
+	INC_DWORD_STAT(STAT_RAC_NumGets);
+
+	FCacheEntryMetadata* Metadata = nullptr;
+	{
+		DECLARE_SCOPE_CYCLE_COUNTER(TEXT("RAC async get time"), STAT_RAC_AsyncGetTime, STATGROUP_RAC);
+		Metadata = CurrentBucket->GetMetadata(CacheKey);
+	}
+
+	if (Metadata == nullptr)
+	{
+		FRuntimeAssetCacheBucketScopeLock Guard(*CurrentBucket);
+		CurrentBucket->AddMetadataEntry(CacheKey, new FCacheEntryMetadata(FDateTime::MaxValue(), 0, 0, CacheKeyName), false);
+	}
+	else
+	{
+		while (Metadata->IsBuilding())
+		{
+			FPlatformProcess::SleepNoStats(0.0f);
+		}
+
+		Metadata = FRuntimeAssetCacheBackend::Get().GetCachedData(BucketName, *CacheKey, Data);
+	}
+
+	ON_SCOPE_EXIT
+	{
+		/** Make sure completed work counter works properly regardless of where function is exited. */
+		GetRuntimeAssetCache().AddToAsyncCompletionCounter(-1);
+	};
+
+	/* Entry found. */
+	if (Metadata
+		/* But was saved with older builder version. */
+		&& Metadata->GetCachedAssetVersion() < CacheBuilder->GetAssetVersion())
+	{
+		/* Pretend entry wasn't found, so it gets rebuilt. */
+		FRuntimeAssetCacheBucketScopeLock Guard(*CurrentBucket);
+		FRuntimeAssetCacheBackend::Get().RemoveCacheEntry(BucketName, *CacheKey);
+		CurrentBucket->AddToCurrentSize(-Metadata->GetCachedAssetSize());
+		CurrentBucket->RemoveMetadataEntry(CacheKey);
+		delete Metadata;
+		Metadata = nullptr;
+	}
+
+	if (Metadata)
+	{
+		check(Data.Num());
+		INC_DWORD_STAT(STAT_RAC_NumCacheHits);
+		FRuntimeAssetCacheBucketScopeLock Guard(*CurrentBucket);
+		Metadata->SetCachedAssetSize(Data.Num());
+		Metadata->SetLastAccessTime(FDateTime::Now());
+		bEntryRetrieved = true;
+		return;
+	}
+
+	if (!CacheBuilder)
+	{
+		// Failed, cleanup data and return false.
+		INC_DWORD_STAT(STAT_RAC_NumFails);
+		Data.Empty();
+		bEntryRetrieved = false;
+		CurrentBucket->RemoveMetadataEntry(CacheKey);
+		return;
+	}
+
+	bool bSuccess;
+	{
+		INC_DWORD_STAT(STAT_RAC_NumBuilds);
+		DECLARE_SCOPE_CYCLE_COUNTER(TEXT("RAC async build time"), STAT_RAC_AsyncBuildTime, STATGROUP_RAC)
+		bSuccess = CacheBuilder->Build(Data);
+	}
+
+	if (!bSuccess)
+	{
+		// Failed, cleanup data and return false.
+		INC_DWORD_STAT(STAT_RAC_NumFails);
+		Data.Empty();
+		CurrentBucket->RemoveMetadataEntry(CacheKey);
+		return;
+	}
+
+	checkf(Data.Num(), TEXT("Size of asset to cache cannot be null. Asset cache key: %s"), *CacheKey);
+	checkf(Data.Num() < CurrentBucket->GetSize(), TEXT("Cached asset is bigger than cache size. Increase cache size or reduce asset size. Asset cache key: %s"), *CacheKey);
+
+	FRuntimeAssetCacheBucketScopeLock Lock(*CurrentBucket);
+
+	// Do we need to make some space in cache?
+	int32 SizeOfSpaceToFree = CurrentBucket->GetCurrentSize() + Data.Num() - CurrentBucket->GetSize();
+	if (SizeOfSpaceToFree > 0)
+	{
+		// Remove oldest entries from cache until we can fit upcoming entry.
+		FreeCacheSpace(BucketName, SizeOfSpaceToFree);
+	}
+
+	{
+		DECLARE_SCOPE_CYCLE_COUNTER(TEXT("RAC async put time"), STAT_RAC_PutTime, STATGROUP_RAC)
+		FDateTime Now = FDateTime::Now();
+		FCacheEntryMetadata* Metadata = CurrentBucket->GetMetadata(*CacheKey);
+		if (Metadata)
+		{
+			Metadata->SetLastAccessTime(Now);
+			if (Metadata->GetCachedAssetSize() == 0)
+			{
+				CurrentBucket->AddToCurrentSize(Data.Num());
+			}
+			Metadata->SetCachedAssetSize(Data.Num());
+			Metadata->SetCachedAssetVersion(CacheBuilder->GetAssetVersion());
+			CurrentBucket->AddMetadataEntry(CacheKey, Metadata, false);
+		}
+		else
+		{
+			Metadata = new FCacheEntryMetadata(Now, Data.Num(), CacheBuilder->GetAssetVersion(), CacheKeyName);
+			CurrentBucket->AddMetadataEntry(CacheKey, Metadata, true);
+		}
+		FRuntimeAssetCacheBackend::Get().PutCachedData(BucketName, *CacheKey, Data, Metadata);
+
+		// Mark that building is finished AFTER putting data into
+		// cache to avoid duplicate builds of the same entry.
+		Metadata->FinishBuilding();
+	}
+
+	bEntryRetrieved = true;
+}
+
+TStatId FRuntimeAssetCacheAsyncWorker::GetStatId()
+{
+	return TStatId();
+}
+
+FString FRuntimeAssetCacheAsyncWorker::SanitizeCacheKey(const TCHAR* CacheKey)
+{
+	FString Output;
+	FString Input(CacheKey);
+	int32 StartValid = 0;
+	int32 NumValid = 0;
+
+	for (int32 i = 0; i < Input.Len(); i++)
+	{
+		if (FChar::IsAlnum(Input[i]) || FChar::IsUnderscore(Input[i]))
+		{
+			NumValid++;
+		}
+		else
+		{
+			// Copy the valid range so far
+			Output += Input.Mid(StartValid, NumValid);
+
+			// Reset valid ranges
+			StartValid = i + 1;
+			NumValid = 0;
+
+			// Replace the invalid character with a special string
+			Output += FString::Printf(TEXT("$%x"), uint32(Input[i]));
+		}
+	}
+
+	// Just return the input if the entire string was valid
+	if (StartValid == 0 && NumValid == Input.Len())
+	{
+		return Input;
+	}
+	else if (NumValid > 0)
+	{
+		// Copy the remaining valid part
+		Output += Input.Mid(StartValid, NumValid);
+	}
+
+	return Output;
+}
+
+FString FRuntimeAssetCacheAsyncWorker::BuildCacheKey(const TCHAR* VersionString, const TCHAR* PluginSpecificCacheKeySuffix)
+{
+	FString UnsanitizedCacheKey = FString(VersionString);
+	UnsanitizedCacheKey.AppendChars(PluginSpecificCacheKeySuffix, FCString::Strlen(PluginSpecificCacheKeySuffix));
+	return SanitizeCacheKey(*UnsanitizedCacheKey);
+}
+
+FString FRuntimeAssetCacheAsyncWorker::BuildCacheKey(FRuntimeAssetCacheBuilderInterface* CacheBuilder)
+{
+	return BuildCacheKey(CacheBuilder->GetVersionString(), *CacheBuilder->GetBuilderSpecificCacheKeySuffix());
+}
+
+void FRuntimeAssetCacheAsyncWorker::FreeCacheSpace(FName Bucket, int32 NumberOfBytesToFree)
+{
+	int32 AccumulatedSize = 0;
+	FRuntimeAssetCacheBucket* CurrentBucket = (*Buckets)[Bucket];
+	while (AccumulatedSize <= NumberOfBytesToFree)
+	{
+		FCacheEntryMetadata* OldestEntry = CurrentBucket->GetOldestEntry();
+		checkf(!OldestEntry->IsBuilding(), TEXT("Cache is trying to remove asset before it finished building. Increase cache size. Asset name: %s, cache size: %d"), *OldestEntry->GetName().ToString(), CurrentBucket->GetSize());
+		AccumulatedSize += OldestEntry->GetCachedAssetSize();
+		FRuntimeAssetCacheBackend::Get().RemoveCacheEntry(Bucket, *OldestEntry->GetName().ToString());
+		(*Buckets)[Bucket]->AddToCurrentSize(-OldestEntry->GetCachedAssetSize());
+		CurrentBucket->RemoveMetadataEntry(OldestEntry->GetName().ToString());
+	}
+}
diff --git a/Engine/Source/Runtime/RuntimeAssetCache/Private/RuntimeAssetCacheAsyncWorker.h b/Engine/Source/Runtime/RuntimeAssetCache/Private/RuntimeAssetCacheAsyncWorker.h
new file mode 100644
--- /dev/null
+++ b/Engine/Source/Runtime/RuntimeAssetCache/Private/RuntimeAssetCacheAsyncWorker.h
@@ -0,0 +1,89 @@
+// Copyright 1998-2015 Epic Games, Inc. All Rights Reserved.
+
+#pragma once
+#include "AsyncWork.h"
+
+/** Stats */
+DECLARE_DWORD_ACCUMULATOR_STAT_EXTERN(TEXT("RAC Num Build"), STAT_RAC_NumBuilds, STATGROUP_RAC, );
+DECLARE_DWORD_ACCUMULATOR_STAT_EXTERN(TEXT("RAC Num Cache Hits"), STAT_RAC_NumCacheHits, STATGROUP_RAC, );
+DECLARE_DWORD_ACCUMULATOR_STAT_EXTERN(TEXT("RAC Num Retrieve fails"), STAT_RAC_NumFails, STATGROUP_RAC, );
+DECLARE_DWORD_ACCUMULATOR_STAT_EXTERN(TEXT("RAC Num Gets"), STAT_RAC_NumGets, STATGROUP_RAC, );
+DECLARE_DWORD_ACCUMULATOR_STAT_EXTERN(TEXT("RAC Num Puts"), STAT_RAC_NumPuts, STATGROUP_RAC, );
+
+/** Forward declarations. */
+class FRuntimeAssetCacheBuilderInterface;
+
+/**
+ * Worker to retrieve entry from cache or build it in case of cache miss.
+ */
+class FRuntimeAssetCacheAsyncWorker : public FNonAbandonableTask
+{
+public:
+	/** Constructor */
+	FRuntimeAssetCacheAsyncWorker(FRuntimeAssetCacheBuilderInterface* InCacheBuilder, TMap<FName, FRuntimeAssetCacheBucket*>* InBuckets);
+
+	/** Async task interface implementation. */
+	void DoWork();
+	TStatId GetStatId();
+	/** End of async task interface implementation. */
+
+	/** Gets serialized cache data. */
+	TArray<uint8>& GetData()
+	{
+		return Data;
+	}
+
+	/**
+	 * Checks if entry was successfully retrieved from cache.
+	 * @return True on success, false otherwise. Success means either cache hit or miss, as long as we got valid data.
+	 */
+	bool RetrievedEntry() const
+	{
+		return bEntryRetrieved;
+	}
+private:
+	/**
+	* Static function to make sure a cache key contains only legal characters by using an escape.
+	* @param CacheKey Cache key to sanitize
+	* @return Sanitized cache key
+	**/
+	static FString SanitizeCacheKey(const TCHAR* CacheKey);
+
+	/**
+	* Static function to build a cache key out of the plugin name, versions and plugin specific info
+	* @param PluginName Name of the runtime asset data type
+	* @param VersionString Version string of the runtime asset data
+	* @param PluginSpecificCacheKeySuffix Whatever is needed to uniquely identify the specific cache entry.
+	* @return Assembled cache key
+	**/
+	static FString BuildCacheKey(const TCHAR* VersionString, const TCHAR* PluginSpecificCacheKeySuffix);
+
+	/**
+	* Static function to build a cache key out of the CacheBuilder.
+	* @param CacheBuilder Builder to create key from.
+	* @return Assembled cache key
+	**/
+	static FString BuildCacheKey(FRuntimeAssetCacheBuilderInterface* CacheBuilder);
+
+	/**
+	* Removes entries from cache until given number of bytes are freed (or more).
+	* @param Bucket Bucket to clean up.
+	* @param SizeOfSpaceToFreeInBytes Number of bytes to free (or more).
+	*/
+	void FreeCacheSpace(FName Bucket, int32 SizeOfSpaceToFreeInBytes);
+
+	/** Cache builder to create cache entry in case of cache miss. */
+	FRuntimeAssetCacheBuilderInterface* CacheBuilder;
+	
+	/** Data to return to caller. */
+	TArray<uint8> Data;
+
+	/** Reference to map of bucket names to their descriptions. */
+	TMap<FName, FRuntimeAssetCacheBucket*>* Buckets;
+
+	/**
+	 * True if successfully retrieved a cache entry. Can be false only when CacheBuilder
+	 * returns false or no CacheBuilder was provided.
+	 */
+	bool bEntryRetrieved;
+};
\ No newline at end of file
diff --git a/Engine/Source/Runtime/RuntimeAssetCache/Private/RuntimeAssetCacheBackend.cpp b/Engine/Source/Runtime/RuntimeAssetCache/Private/RuntimeAssetCacheBackend.cpp
new file mode 100644
--- /dev/null
+++ b/Engine/Source/Runtime/RuntimeAssetCache/Private/RuntimeAssetCacheBackend.cpp
@@ -0,0 +1,58 @@
+// Copyright 1998-2015 Epic Games, Inc. All Rights Reserved.
+
+#include "RuntimeAssetCachePrivatePCH.h"
+#include "RuntimeAssetCacheBackend.h"
+#include "RuntimeAssetCacheFilesystemBackend.h"
+#include "RuntimeAssetCacheEntryMetadata.h"
+
+
+FCacheEntryMetadata* FRuntimeAssetCacheBackend::GetCachedData(const FName Bucket, const TCHAR* CacheKey, TArray<uint8>& OutData)
+{
+	FCacheEntryMetadata* Result = nullptr;
+	FArchive* Ar = CreateReadArchive(Bucket, CacheKey);
+	if (!Ar)
+	{
+		return Result;
+	}
+
+	Result = PreloadMetadata(Ar);
+
+	int64 TotalSize = Ar->TotalSize();
+	int64 CurrentPosition = Ar->Tell();
+	int64 NumberOfBytesToSerialize = TotalSize - CurrentPosition;
+	OutData.Reset();
+	OutData.AddUninitialized(NumberOfBytesToSerialize);
+
+	Ar->Serialize(OutData.GetData(), NumberOfBytesToSerialize);
+	Ar->Close();
+	delete Ar;
+	return Result;
+}
+
+bool FRuntimeAssetCacheBackend::PutCachedData(const FName Bucket, const TCHAR* CacheKey, TArray<uint8>& InData, FCacheEntryMetadata* Metadata)
+{
+	bool bResult = false;
+	FArchive* Ar = CreateWriteArchive(Bucket, CacheKey);
+	if (!Ar)
+	{
+		return bResult;
+	}
+
+	*Ar << *Metadata;
+	Ar->Serialize(InData.GetData(), InData.Num());
+	bResult = Ar->Close();
+	delete Ar;
+	return bResult;
+}
+
+FRuntimeAssetCacheBackend* FRuntimeAssetCacheBackend::CreateBackend()
+{
+	return new FRuntimeAssetCacheFilesystemBackend();
+}
+
+FCacheEntryMetadata* FRuntimeAssetCacheBackend::PreloadMetadata(FArchive* Ar)
+{
+	FCacheEntryMetadata* Result = new FCacheEntryMetadata();
+	*Ar << *Result;
+	return Result;
+}
diff --git a/Engine/Source/Runtime/RuntimeAssetCache/Private/RuntimeAssetCacheBackend.h b/Engine/Source/Runtime/RuntimeAssetCache/Private/RuntimeAssetCacheBackend.h
new file mode 100644
--- /dev/null
+++ b/Engine/Source/Runtime/RuntimeAssetCache/Private/RuntimeAssetCacheBackend.h
@@ -0,0 +1,104 @@
+// Copyright 1998-2015 Epic Games, Inc. All Rights Reserved.
+
+#pragma once
+
+/** Forward declarations. */
+class FCacheEntryMetadata;
+class FRuntimeAssetCacheBucket;
+
+/**
+ * Facade for runtime asset cache backend. Currently only file system backend is available.
+ */
+class FRuntimeAssetCacheBackend : public FNoncopyable
+{
+public:
+	/**
+	 * Singleton getter.
+	 * @return Backend implementation.
+	 */
+	static FRuntimeAssetCacheBackend& Get()
+	{
+		static FRuntimeAssetCacheBackend& Backend = *CreateBackend();
+		return Backend;
+	}
+
+	/**
+	 * Retrieves cached entry.
+	 * @param Bucket Bucket to query.
+	 * @param CacheKey Key of cache entry to retrieve.
+	 * @param OutData Serialized cache entry. Emptied if entry not found.
+	 * @return Metadata descriptor of cache entry on cache hit, nullptr otherwise.
+	 */
+	FCacheEntryMetadata* GetCachedData(FName Bucket, const TCHAR* CacheKey, TArray<uint8>& OutData);
+
+	/**
+	* Puts entry to cache.
+	* @param Bucket Bucket to put entry to.
+	* @param CacheKey Key of cache entry to retrieve.
+	* @param InData Serialized cache entry.
+	* @param Metadata descriptor of cache entry.
+	* @return True if successfully put entry into cache.
+	*/
+	bool PutCachedData(FName Bucket, const TCHAR* CacheKey, TArray<uint8>& InData, FCacheEntryMetadata* Metadata);
+
+	/**
+	* Removes entry from cache.
+	* @param Bucket Bucket to remove entry from.
+	* @param CacheKey Key of cache entry to remove.
+	* @return True if successfully removed entry from cache.
+	*/
+	virtual bool RemoveCacheEntry(FName Bucket, const TCHAR* CacheKey) = 0;
+
+	/**
+	* Removes all entries from cache.
+	* @return True if successfully removed all entries from cache.
+	*/
+	virtual bool ClearCache() = 0;
+
+	/**
+	* Removes all entries from given bucket.
+	* @param Bucket Bucket to clean.
+	* @return True if successfully removed all entries from cache.
+	*/
+	virtual bool ClearCache(FName Bucket) = 0;
+
+	/**
+	* Preloads cache metadata and size for given bucket.
+	* @param BucketName Name of bucket to preload.
+	* @param BucketSize Size of bucket to preload.
+	* @return Preloaded bucket on success, nullptr otherwise.
+	*/
+	virtual FRuntimeAssetCacheBucket* PreLoadBucket(FName BucketName, int32 BucketSize) = 0;
+
+protected:
+	/*
+	* Creates archive to read data from.
+	* @param Bucket Bucket containing cache entry.
+	* @param CacheKey Key to cache entry.
+	* @return Archive storing data.
+	*/
+	virtual FArchive* CreateReadArchive(FName Bucket, const TCHAR* CacheKey) = 0;
+
+	/*
+	* Creates archive to write data into.
+	* @param Bucket Bucket containing cache entry.
+	* @param CacheKey Key to cache entry.
+	* @return Archive storing data.
+	*/
+	virtual FArchive* CreateWriteArchive(FName Bucket, const TCHAR* CacheKey) = 0;
+
+	/*
+	* Preloads metadata from archive.
+	* @param Ar Archive to load metadata from
+	* @return Metadata.
+	*/
+	FCacheEntryMetadata* PreloadMetadata(FArchive* Ar);
+
+private:
+	/**
+	 * Factory method to create actual backend implementation.
+	 * @return Cache backend.
+	 */
+	static FRuntimeAssetCacheBackend* CreateBackend();
+
+};
diff --git a/Engine/Source/Runtime/RuntimeAssetCache/Private/RuntimeAssetCacheBucket.cpp b/Engine/Source/Runtime/RuntimeAssetCache/Private/RuntimeAssetCacheBucket.cpp
new file mode 100644
--- /dev/null
+++ b/Engine/Source/Runtime/RuntimeAssetCache/Private/RuntimeAssetCacheBucket.cpp
@@ -0,0 +1,81 @@
+// Copyright 1998-2015 Epic Games, Inc. All Rights Reserved.
+
+#include "RuntimeAssetCachePrivatePCH.h"
+#include "RuntimeAssetCacheBucket.h"
+#include "RuntimeAssetCacheEntryMetadata.h"
+
+const int32 FRuntimeAssetCacheBucket::DefaultBucketSize = 5 * 1024 * 1024;
+
+FRuntimeAssetCacheBucket::~FRuntimeAssetCacheBucket()
+{
+	for (auto& Entry : CacheMetadata)
+	{
+		delete Entry.Value;
+	}
+}
+
+void FRuntimeAssetCacheBucket::Reset()
+{
+	FScopeLock Lock(&MetadataCriticalSection);
+	for (auto& Kvp : CacheMetadata)
+	{
+		delete Kvp.Value;
+		Kvp.Value = nullptr;
+	}
+	CacheMetadata.Empty();
+	CurrentSize = 0;
+}
+
+FCacheEntryMetadata* FRuntimeAssetCacheBucket::GetMetadata(const FString& Key)
+{
+	FScopeLock Lock(&MetadataCriticalSection);
+	return CacheMetadata.FindOrAdd(Key);
+}
+
+void FRuntimeAssetCacheBucket::RemoveMetadataEntry(const FString& Key)
+{
+	FScopeLock Lock(&MetadataCriticalSection);
+#if UE_BUILD_DEBUG
+	FCacheEntryMetadata* Entry = CacheMetadata.FindOrAdd(Key);
+	if (Entry)
+	{
+		check(!Entry->IsBuilding());
+	}
+#endif
+	CacheMetadata.Remove(Key);
+}
+
+void FRuntimeAssetCacheBucket::AddMetadataEntry(const FString& Key, FCacheEntryMetadata* Value, bool bUpdateSize)
+{
+	check(Value);
+	FScopeLock LockMetadata(&MetadataCriticalSection);
+	FCacheEntryMetadata*& OldValue = CacheMetadata.FindOrAdd(Key);
+	if (bUpdateSize)
+	{
+		FScopeLock LockCurrentSize(&CurrentSizeCriticalSection);
+		if (OldValue)
+		{
+			/** If overwriting, keep proper record of used cache size. */
+			AddToCurrentSize(-OldValue->GetCachedAssetSize());
+		}
+		AddToCurrentSize(Value->GetCachedAssetSize());
+	}
+	OldValue = Value;
+}
+
+FCacheEntryMetadata* FRuntimeAssetCacheBucket::GetOldestEntry()
+{
+	FScopeLock Lock(&MetadataCriticalSection);
+	FDateTime OldestValue = FDateTime::MaxValue();
+	FCacheEntryMetadata* Result = nullptr;
+	for (auto& Kvp : CacheMetadata)
+	{
+		if (Kvp.Value && Kvp.Value->GetLastAccessTime() < OldestValue)
+		{
+			OldestValue = Kvp.Value->GetLastAccessTime();
+			Result = Kvp.Value;
+		}
+	}
+
+	return Result;
+}
diff --git a/Engine/Source/Runtime/RuntimeAssetCache/Private/RuntimeAssetCacheBucket.h b/Engine/Source/Runtime/RuntimeAssetCache/Private/RuntimeAssetCacheBucket.h
new file mode 100644
--- /dev/null
+++ b/Engine/Source/Runtime/RuntimeAssetCache/Private/RuntimeAssetCacheBucket.h
@@ -0,0 +1,111 @@
+// Copyright 1998-2015 Epic Games, Inc. All Rights Reserved.
+
+#pragma once
+
+/** Forward declarations. */
+class FCacheEntryMetadata;
+
+/**
+ * Holds data about cache entries.
+ */
+class FRuntimeAssetCacheBucket
+{
+public:
+	/** Default constructor. */
+	FRuntimeAssetCacheBucket()
+		: Size(DefaultBucketSize)
+		, CurrentSize(0)
+	{ }
+
+	/**
+	 * Constructor
+	 * @param InSize Size of bucket in bytes.
+	 */
+	FRuntimeAssetCacheBucket(int32 InSize)
+		: Size(InSize)
+		, CurrentSize(0)
+	{ }
+
+	/** Destructor. */
+	~FRuntimeAssetCacheBucket();
+
+	/** Gets size in bytes of the bucket. */
+	int32 GetSize() const
+	{
+		return Size;
+	}
+
+	/**
+	 * Cleans up all metadata entries.
+	 */
+	void Reset();
+
+	/**
+	 * Gets metadata entry for given key.
+	 * @param Key Key to get metadata for.
+	 * @return Metadata for given key. Default value if Key is not found.
+	 */
+	FCacheEntryMetadata* GetMetadata(const FString& Key);
+
+	/**
+	* Removes metadata entry for given key.
+	* @param Key Key to remove metadata for.
+	*/
+	void RemoveMetadataEntry(const FString& Key);
+
+	/**
+	* Adds metadata entry for given key.
+	* @param Key Key to add metadata for.
+	* @param Value Metadata to add. If there is an entry for given key, it gets overwritten.
+	* @param bUpdateSize If true, function will update current size of bucket. Otherwise, it's up to caller.
+	*/
+	void AddMetadataEntry(const FString& Key, FCacheEntryMetadata* Value, bool bUpdateSize);
+
+	/**
+	* Gets number of bytes used in cache so far.
+	*/
+	int32 GetCurrentSize()
+	{
+		FScopeLock Lock(&CurrentSizeCriticalSection);
+		return CurrentSize;
+	}
+
+	/**
+	* Adds to number of bytes used in cache so far.
+	* @param Value Number of bytes to add. Can be negative.
+	*/
+	void AddToCurrentSize(int32 Value)
+	{
+		FScopeLock Lock(&CurrentSizeCriticalSection);
+		CurrentSize += Value;
+		check(CurrentSize <= Size);
+	}
+
+	/**
+	* Gets oldest metadata entry in cache.
+	* @return Oldest metadata entry.
+	*/
+	FCacheEntryMetadata* GetOldestEntry();
+
+	/** Default size of bucket in case it wasn't specified in config. */
+	static const int32 DefaultBucketSize;
+
+private:
+	/** Map of cache keys to cache entries metadata. */
+	TMap<FString, FCacheEntryMetadata*> CacheMetadata;
+
+	/** Guard for modifying CacheMetadata (adding/removing elements). Each element has individual guard for modifying it. */
+	FCriticalSection MetadataCriticalSection;
+
+	/** Cache bucket size in bytes. */
+	int32 Size;
+
+	/** Used bucket size in bytes. */
+	int32 CurrentSize;
+
+	/** Guard for modifying current size from multiple threads. */
+	FCriticalSection CurrentSizeCriticalSection;
+
+	/** Helper friend class to guard critical sections on scope basis. */
+	friend class FRuntimeAssetCacheBucketScopeLock;
+};
diff --git a/Engine/Source/Runtime/RuntimeAssetCache/Private/RuntimeAssetCacheBucketScopeLock.h b/Engine/Source/Runtime/RuntimeAssetCache/Private/RuntimeAssetCacheBucketScopeLock.h
new file mode 100644
--- /dev/null
+++ b/Engine/Source/Runtime/RuntimeAssetCache/Private/RuntimeAssetCacheBucketScopeLock.h
@@ -0,0 +1,27 @@
+// Copyright 1998-2015 Epic Games, Inc. All Rights Reserved.
+
+#pragma once
+
+#include "RuntimeAssetCacheBucket.h"
+
+/*
+ * Scope lock guarding RuntimeAssetCache bucket metadata critical section.
+ **/
+class FRuntimeAssetCacheBucketScopeLock
+{
+public:
+	FRuntimeAssetCacheBucketScopeLock(FRuntimeAssetCacheBucket& InBucket)
+		: Bucket(InBucket)
+	{
+		Bucket.MetadataCriticalSection.Lock();
+	}
+
+	~FRuntimeAssetCacheBucketScopeLock()
+	{
+		Bucket.MetadataCriticalSection.Unlock();
+	}
+
+private:
+	/** Bucket to protect. */
+	FRuntimeAssetCacheBucket& Bucket;
+};
\ No newline at end of file
diff --git a/Engine/Source/Runtime/RuntimeAssetCache/Private/RuntimeAssetCacheEntryMetadata.h b/Engine/Source/Runtime/RuntimeAssetCache/Private/RuntimeAssetCacheEntryMetadata.h
new file mode 100644
--- /dev/null
+++ b/Engine/Source/Runtime/RuntimeAssetCache/Private/RuntimeAssetCacheEntryMetadata.h
@@ -0,0 +1,102 @@
+// Copyright 1998-2015 Epic Games, Inc. All Rights Reserved.
+
+#pragma once
+
+/**
+ * Cache entry metadata.
+ */
+class FCacheEntryMetadata
+{
+public:
+	FCacheEntryMetadata(const FDateTime& InLastAccessTime
+		, int32 InCachedAssetSize
+		, int32 InCachedAssetVersion
+		, FName InName)
+		: LastAccessTime(InLastAccessTime)
+		, CachedAssetSize(InCachedAssetSize)
+		, CachedAssetVersion(InCachedAssetVersion)
+		, Name(InName)
+		, bIsBuilding(true)
+	{ }
+
+	FCacheEntryMetadata()
+		: LastAccessTime(FDateTime::Now())
+		, CachedAssetSize(0)
+		, CachedAssetVersion(0)
+		, Name(NAME_None)
+		, bIsBuilding(false)
+	{ }
+
+	int32 GetCachedAssetVersion() const
+	{
+		return CachedAssetVersion;
+	}
+
+	void SetCachedAssetVersion(int32 Version)
+	{
+		CachedAssetVersion = Version;
+	}
+
+	int32 GetCachedAssetSize() const
+	{
+		return CachedAssetSize;
+	}
+
+	void SetCachedAssetSize(int32 Value)
+	{
+		CachedAssetSize = Value;
+	}
+
+	void SetLastAccessTime(FDateTime Value)
+	{
+		LastAccessTime = Value;
+	}
+
+	FDateTime GetLastAccessTime()
+	{
+		return LastAccessTime;
+	}
+
+	FName GetName()
+	{
+		return Name;
+	}
+
+	bool IsBuilding()
+	{
+		return bIsBuilding;
+	}
+
+	void FinishBuilding()
+	{
+		bIsBuilding = false;
+	}
+
+	friend FArchive& operator<<(FArchive& Ar, FCacheEntryMetadata& Metadata);
+
+private:
+	FDateTime LastAccessTime;
+	int32 CachedAssetSize;
+	int32 CachedAssetVersion;
+	FName Name;
+	FThreadSafeBool bIsBuilding;
+};
+
+FArchive& operator<<(FArchive& Ar, FCacheEntryMetadata& Metadata)
+{
+	Ar << Metadata.CachedAssetSize;
+	Ar << Metadata.CachedAssetVersion;
+	FString String;
+	if (Ar.IsLoading())
+	{
+		Ar << String;
+		Metadata.Name = FName(*String);
+	}
+	else if (Ar.IsSaving())
+	{
+		Metadata.Name.ToString(String);
+		Ar << String; 
+	}
+
+	return Ar;
+}
diff --git a/Engine/Source/Runtime/RuntimeAssetCache/Private/RuntimeAssetCacheFilesystemBackend.cpp b/Engine/Source/Runtime/RuntimeAssetCache/Private/RuntimeAssetCacheFilesystemBackend.cpp
new file mode 100644
--- /dev/null
+++ b/Engine/Source/Runtime/RuntimeAssetCache/Private/RuntimeAssetCacheFilesystemBackend.cpp
@@ -0,0 +1,70 @@
+// Copyright 1998-2015 Epic Games, Inc. All Rights Reserved.
+
+#pragma once
+#include "RuntimeAssetCachePrivatePCH.h"
+#include "RuntimeAssetCacheFilesystemBackend.h"
+#include "RuntimeAssetCacheEntryMetadata.h"
+#include "RuntimeAssetCacheBucket.h"
+
+const FString FRuntimeAssetCacheFilesystemBackend::PathToRAC = TEXT("../../RuntimeAssetCache/");
+
+FArchive* FRuntimeAssetCacheFilesystemBackend::CreateReadArchive(FName Bucket, const TCHAR* CacheKey)
+{
+	FString Path = FPaths::Combine(*PathToRAC, *Bucket.ToString(), CacheKey);
+	return IFileManager::Get().CreateFileReader(*Path);
+}
+
+FArchive* FRuntimeAssetCacheFilesystemBackend::CreateWriteArchive(FName Bucket, const TCHAR* CacheKey)
+{
+	FString Path = FPaths::Combine(*PathToRAC, *Bucket.ToString(), CacheKey);
+	return IFileManager::Get().CreateFileWriter(*Path);
+}
+
+bool FRuntimeAssetCacheFilesystemBackend::RemoveCacheEntry(const FName Bucket, const TCHAR* CacheKey)
+{
+	FString Path = FPaths::Combine(*PathToRAC, *Bucket.ToString(), CacheKey);
+	return IFileManager::Get().Delete(*Path);
+}
+
+bool FRuntimeAssetCacheFilesystemBackend::ClearCache()
+{
+	return IFileManager::Get().DeleteDirectory(*PathToRAC, false, true);
+}
+
+bool FRuntimeAssetCacheFilesystemBackend::ClearCache(FName Bucket)
+{
+	return IFileManager::Get().DeleteDirectory(*FPaths::Combine(*PathToRAC, *Bucket.ToString()), false, true);
+}
+
+FRuntimeAssetCacheBucket* FRuntimeAssetCacheFilesystemBackend::PreLoadBucket(FName BucketName, int32 BucketSize)
+{
+	FString Path = FPaths::Combine(*PathToRAC, *BucketName.ToString());
+	FRuntimeAssetCacheBucket* Result = new FRuntimeAssetCacheBucket(BucketSize);
+
+	class FRuntimeAssetCacheFilesystemBackendDirectoryVisitor : public IPlatformFile::FDirectoryVisitor
+	{
+	public:
+		FRuntimeAssetCacheFilesystemBackendDirectoryVisitor(FRuntimeAssetCacheBucket* InBucket, FName InBucketName, FRuntimeAssetCacheFilesystemBackend* InBackend)
+			: Bucket(InBucket)
+			, BucketName(InBucketName)
+			, Backend(InBackend)
+		{ }
+
+		virtual bool Visit(const TCHAR* FilenameOrDirectory, bool bIsDirectory)
+		{
+			FArchive* Ar = Backend->CreateReadArchive(BucketName, FilenameOrDirectory);
+			FCacheEntryMetadata* Metadata = Backend->PreloadMetadata(Ar);
+			Bucket->AddMetadataEntry(*FPaths::GetBaseFilename(FilenameOrDirectory), Metadata, true);
+			delete Ar;
+			return true;
+		}
+
+	private:
+		FRuntimeAssetCacheBucket* Bucket;
+		FName BucketName;
+		FRuntimeAssetCacheFilesystemBackend* Backend;
+	} Visitor(Result, BucketName, this);
+	IFileManager::Get().IterateDirectory(*Path, Visitor);
+
+	return Result;
+}
diff --git a/Engine/Source/Runtime/RuntimeAssetCache/Private/RuntimeAssetCacheFilesystemBackend.h b/Engine/Source/Runtime/RuntimeAssetCache/Private/RuntimeAssetCacheFilesystemBackend.h
new file mode 100644
--- /dev/null
+++ b/Engine/Source/Runtime/RuntimeAssetCache/Private/RuntimeAssetCacheFilesystemBackend.h
@@ -0,0 +1,27 @@
+// Copyright 1998-2015 Epic Games, Inc. All Rights Reserved.
+
+#pragma once
+
+#include "RuntimeAssetCacheBackend.h"
+
+/**
+ * File system implementation of runtime asset cache backend.
+ */
+class FRuntimeAssetCacheFilesystemBackend : public FRuntimeAssetCacheBackend
+{
+public:
+	/** FRuntimeAssetCacheBackend interface implementation. */
+	virtual bool RemoveCacheEntry(const FName Bucket, const TCHAR* CacheKey) override;
+	virtual bool ClearCache() override;
+	virtual bool ClearCache(FName Bucket) override;
+	virtual FRuntimeAssetCacheBucket* PreLoadBucket(FName BucketName, int32 BucketSize) override;
+protected:
+	virtual FArchive* CreateReadArchive(FName Bucket, const TCHAR* CacheKey) override;
+	virtual FArchive* CreateWriteArchive(FName Bucket, const TCHAR* CacheKey) override;
+	/** End of FRuntimeAssetCacheBackend interface implementation. */
+
+private:
+
+	/** File system path to runtime asset cache. */
+	static const FString PathToRAC;
+};
diff --git a/Engine/Source/Runtime/RuntimeAssetCache/Private/RuntimeAssetCacheModule.cpp b/Engine/Source/Runtime/RuntimeAssetCache/Private/RuntimeAssetCacheModule.cpp
new file mode 100644
--- /dev/null
+++ b/Engine/Source/Runtime/RuntimeAssetCache/Private/RuntimeAssetCacheModule.cpp
@@ -0,0 +1,35 @@
+// Copyright 1998-2015 Epic Games, Inc. All Rights Reserved.
+
+#include "RuntimeAssetCachePrivatePCH.h"
+#include "RuntimeAssetCacheModule.h"
+#include "RuntimeAssetCache.h"
+#include "ModuleManager.h"
+
+DEFINE_LOG_CATEGORY(RuntimeAssetCache)
+
+FRuntimeAssetCacheInterface& GetRuntimeAssetCache()
+{
+	static FRuntimeAssetCacheModuleInterface& Module = FModuleManager::LoadModuleChecked<FRuntimeAssetCacheModuleInterface>("RuntimeAssetCache");
+	static FRuntimeAssetCacheInterface& Interface = Module.GetRuntimeAssetCache();
+	return Interface;
+}
+
+/**
+ * Class implementing RuntimeAssetCache module interface. This is split
+ * to solve circular dependency when building UHT.
+ */
+class FRuntimeAssetCacheModule : public FRuntimeAssetCacheModuleInterface
+{
+public:
+	/**
+	* Gets runtime asset cache.
+	* @return Runtime asset cache.
+	*/
+	virtual FRuntimeAssetCacheInterface& GetRuntimeAssetCache() override
+	{
+		static FRuntimeAssetCache RuntimeAssetCache = FRuntimeAssetCache();
+		return RuntimeAssetCache;
+	}
+};
+
+IMPLEMENT_MODULE(FRuntimeAssetCacheModule, RuntimeAssetCache);
diff --git a/Engine/Source/Runtime/RuntimeAssetCache/Private/RuntimeAssetCachePrivatePCH.h b/Engine/Source/Runtime/RuntimeAssetCache/Private/RuntimeAssetCachePrivatePCH.h
new file mode 100644
--- /dev/null
+++ b/Engine/Source/Runtime/RuntimeAssetCache/Private/RuntimeAssetCachePrivatePCH.h
@@ -0,0 +1,11 @@
+// Copyright 1998-2015 Epic Games, Inc. All Rights Reserved.
+
+#pragma once
+
+#include "Core.h"
+
+/** Stats */
+DECLARE_STATS_GROUP(TEXT("RAC"), STATGROUP_RAC, STATCAT_Advanced);
+
+/** Logs */
+DECLARE_LOG_CATEGORY_EXTERN(RuntimeAssetCache, Log, All);
diff --git a/Engine/Source/Runtime/RuntimeAssetCache/Public/RuntimeAssetCacheInterface.h b/Engine/Source/Runtime/RuntimeAssetCache/Public/RuntimeAssetCacheInterface.h
new file mode 100644
--- /dev/null
+++ b/Engine/Source/Runtime/RuntimeAssetCache/Public/RuntimeAssetCacheInterface.h
@@ -0,0 +1,79 @@
+// Copyright 1998-2015 Epic Games, Inc. All Rights Reserved.
+
+#pragma once
+
+/** Forward declarations */
+class FRuntimeAssetCacheBuilderInterface;
+
+/**
+* Interface for the Runtime Asset Cache. Cache is split into buckets to cache various assets separately.
+* Bucket names and maximum sizes are configure via Engine.ini config file using the following syntax
+* [RuntimeAssetCache]
+* +BucketConfigs=(Name="<Plugin name>", Size=<Maximum bucket size in bytes>)
+* This API is fully thread safe.
+**/
+class FRuntimeAssetCacheInterface
+{
+public:
+	/**
+	* Synchronously gets value from cache. If value is not found, builds entry using CacheBuilder and updates cache.
+	* @param CacheBuilder Builder to produce cache key and in the event of a miss.
+	* @param OutData Reference to serialized data.
+	* @return true if the data was retrieved from the cache or the deriver built the data successfully. false can only occur if the plugin returns false.
+	**/
+	virtual bool GetSynchronous(FRuntimeAssetCacheBuilderInterface* CacheBuilder, TArray<uint8>& OutData) = 0;
+
+	/**
+	* Asynchronously checks the cache. If value is not found, builds entry using CacheBuilder and updates cache.
+	* @param CacheBuilder Builder to produce cache key and in the event of a miss, return the data.
+	* @return Handle to worker.
+	**/
+	virtual uint32 GetAsynchronous(FRuntimeAssetCacheBuilderInterface* CacheBuilder) = 0;
+
+	/**
+	* Gets cache size.
+	* @param Bucket to check.
+	* @return Maximum allowed bucket size.
+	**/
+	virtual int32 GetCacheSize(FName Bucket) const = 0;
+
+	/**
+	* Removes all cache entries.
+	* @return true if cache was successfully cleaned.
+	*/
+	virtual bool ClearCache() = 0;
+
+	/**
+	* Removes all cache from given bucket.
+	* @param Bucket Bucket to clean.
+	* @return true if cache was successfully cleaned.
+	*/
+	virtual bool ClearCache(FName Bucket) = 0;
+
+	/**
+	* Waits until worker finishes execution.
+	* @param Handle Worker handle to wait for.
+	*/
+	virtual void WaitAsynchronousCompletion(uint32 Handle) = 0;
+
+	/**
+	* Gets asynchronous query results.
+	* @param Handle Worker handle to check.
+	* @param OutData Reference to serialized cache entry.
+	* @return True if successfully retrieved data, false otherwise.
+	*/
+	virtual bool GetAsynchronousResults(uint32 Handle, TArray<uint8>& OutData) = 0;
+
+	/**
+	* Checks if worker finished execution.
+	* @param Handle Worker handle to check.
+	* @return True if execution finished, false otherwise.
+	*/
+	virtual bool PollAsynchronousCompletion(uint32 Handle) = 0;
+
+	/**
+	 * Adds a number from the thread safe counter which tracks outstanding async requests. This is used to ensure everything is complete prior to shutdown.
+	 * @param Value to add to counter. Can be negative.
+	 */
+	virtual void AddToAsyncCompletionCounter(int32 Addend) = 0;
+};
diff --git a/Engine/Source/Runtime/RuntimeAssetCache/Public/RuntimeAssetCacheModule.h b/Engine/Source/Runtime/RuntimeAssetCache/Public/RuntimeAssetCacheModule.h
new file mode 100644
--- /dev/null
+++ b/Engine/Source/Runtime/RuntimeAssetCache/Public/RuntimeAssetCacheModule.h
@@ -0,0 +1,22 @@
+// Copyright 1998-2015 Epic Games, Inc. All Rights Reserved.
+
+#pragma once
+#include "Modules/ModuleInterface.h"
+
+/** Forward declarations. */
+class FRuntimeAssetCacheInterface;
+
+RUNTIMEASSETCACHE_API FRuntimeAssetCacheInterface& GetRuntimeAssetCache();
+
+/**
+* Module for the RuntimeAssetCache.
+*/
+class FRuntimeAssetCacheModuleInterface : public IModuleInterface
+{
+public:
+	/**
+	 * Gets runtime asset cache.
+	 * @return Runtime asset cache.
+	 */
+	virtual FRuntimeAssetCacheInterface& GetRuntimeAssetCache() = 0;
+};
diff --git a/Engine/Source/Runtime/RuntimeAssetCache/Public/RuntimeAssetCachePluginInterface.h b/Engine/Source/Runtime/RuntimeAssetCache/Public/RuntimeAssetCachePluginInterface.h
new file mode 100644
--- /dev/null
+++ b/Engine/Source/Runtime/RuntimeAssetCache/Public/RuntimeAssetCachePluginInterface.h
@@ -0,0 +1,47 @@
+// Copyright 1998-2015 Epic Games, Inc. All Rights Reserved.
+
+#pragma once
+
+/**
+* Interface for runtime asset cache cache builders.
+* This API will be called concurrently.
+**/
+class FRuntimeAssetCacheBuilderInterface
+{
+public:
+	/** Virtual destructor. */
+	virtual ~FRuntimeAssetCacheBuilderInterface()
+	{ }
+
+	/**
+	* Get the builder type name. This is used to categorize cached data to buckets.
+	* @return Type name of builder.
+	**/
+	virtual const TCHAR* GetTypeName() const = 0;
+
+	/**
+	* Get the version of cache builder, this is used as part of the cache key. This is supposed to
+	* be a guid string ( ex. "69C8C8A6-A9F8-4EFC-875C-CFBB72E66486" )
+	* @return Version string of cache builder
+	**/
+	virtual const TCHAR* GetVersionString() const = 0;
+
+	/**
+	* Returns the builder specific part of the cache key. This must be a alphanumeric+underscore
+	* @return Version number of the builder.
+	**/
+	virtual FString GetBuilderSpecificCacheKeySuffix() const = 0;
+
+	/**
+	* Does the work of creating serialized cache entry.
+	* @param OutData Array of bytes to fill in with the result data
+	* @return true if successful, in the event of failure the cache is not updated and failure is propagated to the original caller.
+	**/
+	virtual bool Build(TArray<uint8>& OutData) = 0;
+
+	/**
+	 * Gets asset version to rebuild cache if cached asset is too old.
+	 * @return Asset version.
+	 */
+	virtual int32 GetAssetVersion() = 0;
+};
diff --git a/Engine/Source/Runtime/RuntimeAssetCache/RuntimeAssetCache.Build.cs b/Engine/Source/Runtime/RuntimeAssetCache/RuntimeAssetCache.Build.cs
new file mode 100644
--- /dev/null
+++ b/Engine/Source/Runtime/RuntimeAssetCache/RuntimeAssetCache.Build.cs
@@ -0,0 +1,12 @@
+// Copyright 1998-2015 Epic Games, Inc. All Rights Reserved.
+
+using UnrealBuildTool;
+using System.IO;
+
+public class RuntimeAssetCache : ModuleRules
+{
+	public RuntimeAssetCache(TargetInfo Target)
+	{
+		PrivateDependencyModuleNames.Add("Core");
+	}
+}
