Fix mismatched derived function declaration.

diff --git a/Engine/Source/Runtime/RuntimeAssetCache/Private/RuntimeAssetCacheFilesystemBackend.h b/Engine/Source/Runtime/RuntimeAssetCache/Private/RuntimeAssetCacheFilesystemBackend.h
--- a/Engine/Source/Runtime/RuntimeAssetCache/Private/RuntimeAssetCacheFilesystemBackend.h
+++ b/Engine/Source/Runtime/RuntimeAssetCache/Private/RuntimeAssetCacheFilesystemBackend.h
@@ -11,7 +11,7 @@ class FRuntimeAssetCacheFilesystemBackend : public FRuntimeAssetCacheBackend
 {
 public:
 	/** FRuntimeAssetCacheBackend interface implementation. */
-	virtual bool RemoveCacheEntry(const FName Bucket, const TCHAR* CacheKey) override;
+	virtual bool RemoveCacheEntry(FName Bucket, const TCHAR* CacheKey) override;
 	virtual bool ClearCache() override;
 	virtual bool ClearCache(FName Bucket) override;
 	virtual FRuntimeAssetCacheBucket* PreLoadBucket(FName BucketName, int32 BucketSize) override;
