Small fixes to Runtime Asset Cache.

diff --git a/Engine/Source/Runtime/RuntimeAssetCache/Private/RuntimeAssetCacheAsyncWorker.cpp b/Engine/Source/Runtime/RuntimeAssetCache/Private/RuntimeAssetCacheAsyncWorker.cpp
--- a/Engine/Source/Runtime/RuntimeAssetCache/Private/RuntimeAssetCacheAsyncWorker.cpp
+++ b/Engine/Source/Runtime/RuntimeAssetCache/Private/RuntimeAssetCacheAsyncWorker.cpp
@@ -226,3 +226,22 @@ void FRuntimeAssetCacheAsyncWorker::FreeCacheSpace(FName Bucket, int32 NumberOfB
 		CurrentBucket->RemoveMetadataEntry(OldestEntry->GetName().ToString());
 	}
 }
+
+FArchive& operator<<(FArchive& Ar, FCacheEntryMetadata& Metadata)
+{
+	Ar << Metadata.CachedAssetSize;
+	Ar << Metadata.CachedAssetVersion;
+	FString String;
+	if (Ar.IsLoading())
+	{
+		Ar << String;
+		Metadata.Name = FName(*String);
+	}
+	else if (Ar.IsSaving())
+	{
+		Metadata.Name.ToString(String);
+		Ar << String;
+	}
+
+	return Ar;
+}
\ No newline at end of file
diff --git a/Engine/Source/Runtime/RuntimeAssetCache/Private/RuntimeAssetCacheAsyncWorker.h b/Engine/Source/Runtime/RuntimeAssetCache/Private/RuntimeAssetCacheAsyncWorker.h
--- a/Engine/Source/Runtime/RuntimeAssetCache/Private/RuntimeAssetCacheAsyncWorker.h
+++ b/Engine/Source/Runtime/RuntimeAssetCache/Private/RuntimeAssetCacheAsyncWorker.h
@@ -12,6 +12,7 @@ DECLARE_DWORD_ACCUMULATOR_STAT_EXTERN(TEXT("RAC Num Puts"), STAT_RAC_NumPuts, ST
 
 /** Forward declarations. */
 class FRuntimeAssetCacheBuilderInterface;
+class FRuntimeAssetCacheBucket;
 
 /**
  * Worker to retrieve entry from cache or build it in case of cache miss.
diff --git a/Engine/Source/Runtime/RuntimeAssetCache/Private/RuntimeAssetCacheEntryMetadata.h b/Engine/Source/Runtime/RuntimeAssetCache/Private/RuntimeAssetCacheEntryMetadata.h
--- a/Engine/Source/Runtime/RuntimeAssetCache/Private/RuntimeAssetCacheEntryMetadata.h
+++ b/Engine/Source/Runtime/RuntimeAssetCache/Private/RuntimeAssetCacheEntryMetadata.h
@@ -82,21 +82,4 @@ private:
 	FThreadSafeBool bIsBuilding;
 };
 
-FArchive& operator<<(FArchive& Ar, FCacheEntryMetadata& Metadata)
-{
-	Ar << Metadata.CachedAssetSize;
-	Ar << Metadata.CachedAssetVersion;
-	FString String;
-	if (Ar.IsLoading())
-	{
-		Ar << String;
-		Metadata.Name = FName(*String);
-	}
-	else if (Ar.IsSaving())
-	{
-		Metadata.Name.ToString(String);
-		Ar << String; 
-	}
-
-	return Ar;
-}
+FArchive& operator<<(FArchive& Ar, FCacheEntryMetadata& Metadata);
