Minor tweaks to runtime asset cache.

diff --git a/Engine/Source/Runtime/RuntimeAssetCache/Private/RuntimeAssetCacheAsyncWorker.cpp b/Engine/Source/Runtime/RuntimeAssetCache/Private/RuntimeAssetCacheAsyncWorker.cpp
--- a/Engine/Source/Runtime/RuntimeAssetCache/Private/RuntimeAssetCacheAsyncWorker.cpp
+++ b/Engine/Source/Runtime/RuntimeAssetCache/Private/RuntimeAssetCacheAsyncWorker.cpp
@@ -209,7 +209,7 @@ FString FRuntimeAssetCacheAsyncWorker::BuildCacheKey(const TCHAR* VersionString,
 
 FString FRuntimeAssetCacheAsyncWorker::BuildCacheKey(FRuntimeAssetCacheBuilderInterface* CacheBuilder)
 {
-	return BuildCacheKey(CacheBuilder->GetVersionString(), *CacheBuilder->GetBuilderSpecificCacheKeySuffix());
+	return BuildCacheKey(CacheBuilder->GetBuilderName(), *CacheBuilder->GetAssetUniqueName());
 }
 
 void FRuntimeAssetCacheAsyncWorker::FreeCacheSpace(FName Bucket, int32 NumberOfBytesToFree)
diff --git a/Engine/Source/Runtime/RuntimeAssetCache/Public/RuntimeAssetCachePluginInterface.h b/Engine/Source/Runtime/RuntimeAssetCache/Public/RuntimeAssetCachePluginInterface.h
--- a/Engine/Source/Runtime/RuntimeAssetCache/Public/RuntimeAssetCachePluginInterface.h
+++ b/Engine/Source/Runtime/RuntimeAssetCache/Public/RuntimeAssetCachePluginInterface.h
@@ -20,17 +20,16 @@ public:
 	virtual const TCHAR* GetTypeName() const = 0;
 
 	/**
-	* Get the version of cache builder, this is used as part of the cache key. This is supposed to
-	* be a guid string ( ex. "69C8C8A6-A9F8-4EFC-875C-CFBB72E66486" )
-	* @return Version string of cache builder
+	* Get the name of cache builder, this is used as part of the cache key.
+	* @return Name of cache builder
 	**/
-	virtual const TCHAR* GetVersionString() const = 0;
+	virtual const TCHAR* GetBuilderName() const = 0;
 
 	/**
 	* Returns the builder specific part of the cache key. This must be a alphanumeric+underscore
 	* @return Version number of the builder.
 	**/
-	virtual FString GetBuilderSpecificCacheKeySuffix() const = 0;
+	virtual FString GetAssetUniqueName() const = 0;
 
 	/**
 	* Does the work of creating serialized cache entry.
diff --git a/Engine/Source/Runtime/RuntimeAssetCache/RuntimeAssetCache.Build.cs b/Engine/Source/Runtime/RuntimeAssetCache/RuntimeAssetCache.Build.cs
--- a/Engine/Source/Runtime/RuntimeAssetCache/RuntimeAssetCache.Build.cs
+++ b/Engine/Source/Runtime/RuntimeAssetCache/RuntimeAssetCache.Build.cs
@@ -8,5 +8,7 @@ public class RuntimeAssetCache : ModuleRules
 	public RuntimeAssetCache(TargetInfo Target)
 	{
 		PrivateDependencyModuleNames.Add("Core");
+
+		SharedPCHHeaderFile = "Runtime/RuntimeAssetCache/Public/RuntimeAssetCachePublicPCH.h";
 	}
 }
