Fix RuntimeAssetCache compilation issues

diff --git a/Engine/Source/Runtime/RuntimeAssetCache/RuntimeAssetCache.Build.cs b/Engine/Source/Runtime/RuntimeAssetCache/RuntimeAssetCache.Build.cs
--- a/Engine/Source/Runtime/RuntimeAssetCache/RuntimeAssetCache.Build.cs
+++ b/Engine/Source/Runtime/RuntimeAssetCache/RuntimeAssetCache.Build.cs
@@ -8,7 +8,5 @@ public class RuntimeAssetCache : ModuleRules
 	public RuntimeAssetCache(TargetInfo Target)
 	{
 		PrivateDependencyModuleNames.Add("Core");
-
-		SharedPCHHeaderFile = "Runtime/RuntimeAssetCache/Public/RuntimeAssetCachePublicPCH.h";
 	}
 }
