Make RuntimeAssetCache blueprint accessible.

diff --git a/Engine/Source/Runtime/Core/Public/Serialization/Archive.h b/Engine/Source/Runtime/Core/Public/Serialization/Archive.h
--- a/Engine/Source/Runtime/Core/Public/Serialization/Archive.h
+++ b/Engine/Source/Runtime/Core/Public/Serialization/Archive.h
@@ -300,6 +300,101 @@ protected:
 	bool	bFreeOnClose;
 };
 
+
+/**
+* Similar to FMemoryWriter, but able to internally
+* manage the memory for the buffer.
+*/
+class FBufferWriter : public FArchive
+{
+public:
+	/**
+	* Constructor
+	*
+	* @param Data Buffer to use as the source data to read from
+	* @param Size Size of Data
+	* @param bInFreeOnClose If true, Data will be FMemory::Free'd when this archive is closed
+	* @param bIsPersistent Uses this value for ArIsPersistent
+	*/
+	FBufferWriter(void* Data, int64 Size, bool bInFreeOnClose, bool bIsPersistent = false)
+		: WriterData(Data)
+		, WriterPos(0)
+		, WriterSize(Size)
+		, bFreeOnClose(bInFreeOnClose)
+	{
+		ArIsSaving = true;
+		ArIsPersistent = bIsPersistent;
+	}
+
+	~FBufferWriter()
+	{
+		Close();
+	}
+	bool Close()
+	{
+		if (bFreeOnClose)
+		{
+			FMemory::Free(WriterData);
+			WriterData = nullptr;
+		}
+		return !ArIsError;
+	}
+	void Serialize(void* Data, int64 Num)
+	{
+		const int64 NumBytesToAdd = WriterPos + Num - WriterSize;
+		if (NumBytesToAdd > 0)
+		{
+			const int64 NewArrayCount = WriterSize + NumBytesToAdd;
+			if (NewArrayCount >= MAX_int32)
+			{
+				UE_LOG(LogSerialization, Fatal, TEXT("FBufferWriter does not support data larger than 2GB. Archive name: %s."), *GetArchiveName());
+			}
+
+			FMemory::Realloc(WriterData, NewArrayCount);
+			WriterSize = NewArrayCount;
+		}
+
+		check(WriterPos >= 0);
+		check((WriterPos + Num) <= WriterSize);
+		FMemory::Memcpy((uint8*)WriterData + WriterPos, Data, Num);
+		WriterPos += Num;
+	}
+	int64 Tell()
+	{
+		return WriterPos;
+	}
+	int64 TotalSize()
+	{
+		return WriterSize;
+	}
+	void Seek(int64 InPos)
+	{
+		check(InPos >= 0);
+		check(InPos <= WriterSize);
+		WriterPos = InPos;
+	}
+	bool AtEnd()
+	{
+		return WriterPos >= WriterSize;
+	}
+	/**
+	* Returns the name of the Archive.  Useful for getting the name of the package a struct or object
+	* is in when a loading error occurs.
+	*
+	* This is overridden for the specific Archive Types
+	**/
+	virtual FString GetArchiveName() const { return TEXT("FBufferWriter"); }
+
+	void* GetWriterData()
+	{
+		return WriterData;
+	}
+protected:
+	void* WriterData;
+	int64 WriterPos;
+	int64 WriterSize;
+	bool bFreeOnClose;
+};
 /*----------------------------------------------------------------------------
 	FArchiveSaveCompressedProxy.
 ----------------------------------------------------------------------------*/
diff --git a/Engine/Source/Runtime/Engine/Classes/Engine/World.h b/Engine/Source/Runtime/Engine/Classes/Engine/World.h
--- a/Engine/Source/Runtime/Engine/Classes/Engine/World.h
+++ b/Engine/Source/Runtime/Engine/Classes/Engine/World.h
@@ -3252,6 +3252,9 @@ public:
 	// delegate for generating world asset registry tags so project/game scope can add additional tags for filtering levels in their UI, etc
 	DECLARE_MULTICAST_DELEGATE_TwoParams(FWorldGetAssetTags, const UWorld*, TArray<UObject::FAssetRegistryTag>&);
 
+	DECLARE_MULTICAST_DELEGATE_TwoParams(FOnWorldTickStart, ELevelTick, float);
+	static FOnWorldTickStart OnWorldTickStart;
+
 	// Callback for world creation
 	static FWorldEvent OnPostWorldCreation;
 	
diff --git a/Engine/Source/Runtime/Engine/Private/LevelTick.cpp b/Engine/Source/Runtime/Engine/Private/LevelTick.cpp
--- a/Engine/Source/Runtime/Engine/Private/LevelTick.cpp
+++ b/Engine/Source/Runtime/Engine/Private/LevelTick.cpp
@@ -1017,6 +1017,8 @@ void UWorld::Tick( ELevelTick TickType, float DeltaSeconds )
 		return;
 	}
 
+	FWorldDelegates::OnWorldTickStart.Broadcast(TickType, DeltaSeconds);
+
 #if LOG_DETAILED_PATHFINDING_STATS
 	GDetailedPathFindingStats.Reset();
 #endif
diff --git a/Engine/Source/Runtime/Engine/Private/World.cpp b/Engine/Source/Runtime/Engine/Private/World.cpp
--- a/Engine/Source/Runtime/Engine/Private/World.cpp
+++ b/Engine/Source/Runtime/Engine/Private/World.cpp
@@ -108,6 +108,7 @@ FWorldDelegates::FOnLevelChanged FWorldDelegates::LevelAddedToWorld;
 FWorldDelegates::FOnLevelChanged FWorldDelegates::LevelRemovedFromWorld;
 FWorldDelegates::FLevelOffsetEvent FWorldDelegates::PostApplyLevelOffset;
 FWorldDelegates::FWorldGetAssetTags FWorldDelegates::GetAssetTags;
+FWorldDelegates::FOnWorldTickStart FWorldDelegates::OnWorldTickStart;
 #if WITH_EDITOR
 FWorldDelegates::FRefreshLevelScriptActionsEvent FWorldDelegates::RefreshLevelScriptActions;
 #endif // WITH_EDITOR
diff --git a/Engine/Source/Runtime/RuntimeAssetCache/Private/RuntimeAssetCache.cpp b/Engine/Source/Runtime/RuntimeAssetCache/Private/RuntimeAssetCache.cpp
--- a/Engine/Source/Runtime/RuntimeAssetCache/Private/RuntimeAssetCache.cpp
+++ b/Engine/Source/Runtime/RuntimeAssetCache/Private/RuntimeAssetCache.cpp
@@ -7,6 +7,8 @@
 #include "RuntimeAssetCacheBucket.h"
 #include "RuntimeAssetCacheAsyncWorker.h"
 #include "RuntimeAssetCacheEntryMetadata.h"
+#include "RuntimeAssetCacheInterface.h"
+#include "RuntimeAssetCacheBPHooks.h"
 
 DEFINE_STAT(STAT_RAC_ASyncWaitTime);
 
@@ -45,7 +47,7 @@ int32 FRuntimeAssetCache::GetCacheSize(FName Bucket) const
 	return Buckets[Bucket]->GetSize();
 }
 
-uint32 FRuntimeAssetCache::GetAsynchronous(FRuntimeAssetCacheBuilderInterface* CacheBuilder)
+int32 FRuntimeAssetCache::GetAsynchronous(IRuntimeAssetCacheBuilder* CacheBuilder, const FOnRuntimeAssetCacheAsyncComplete& OnComplete)
 {
 	int32 Handle = GetNextHandle();
 
@@ -55,7 +57,7 @@ uint32 FRuntimeAssetCache::GetAsynchronous(FRuntimeAssetCacheBuilderInterface* C
 	/** Make sure task isn't processed twice. */
 	check(!PendingTasks.Contains(Handle));
 
-	FAsyncTask<FRuntimeAssetCacheAsyncWorker>* AsyncTask = new FAsyncTask<FRuntimeAssetCacheAsyncWorker>(CacheBuilder, &Buckets);
+	FAsyncTask<FRuntimeAssetCacheAsyncWorker>* AsyncTask = new FAsyncTask<FRuntimeAssetCacheAsyncWorker>(CacheBuilder, &Buckets, Handle, OnComplete);
 
 	{
 		FScopeLock ScopeLock(&SynchronizationObject);
@@ -65,14 +67,17 @@ uint32 FRuntimeAssetCache::GetAsynchronous(FRuntimeAssetCacheBuilderInterface* C
 	AsyncTask->StartBackgroundTask();
 	return Handle;
 }
+int32 FRuntimeAssetCache::GetAsynchronous(IRuntimeAssetCacheBuilder* CacheBuilder)
+{
+	return GetAsynchronous(CacheBuilder, FOnRuntimeAssetCacheAsyncComplete());
+}
 
-bool FRuntimeAssetCache::GetSynchronous(FRuntimeAssetCacheBuilderInterface* CacheBuilder, TArray<uint8>& OutData)
+void* FRuntimeAssetCache::GetSynchronous(IRuntimeAssetCacheBuilder* CacheBuilder)
 {
-	FAsyncTask<FRuntimeAssetCacheAsyncWorker>* AsyncTask = new FAsyncTask<FRuntimeAssetCacheAsyncWorker>(CacheBuilder, &Buckets);
+	FAsyncTask<FRuntimeAssetCacheAsyncWorker>* AsyncTask = new FAsyncTask<FRuntimeAssetCacheAsyncWorker>(CacheBuilder, &Buckets, -1, FOnRuntimeAssetCacheAsyncComplete());
 	AddToAsyncCompletionCounter(1);
 	AsyncTask->StartSynchronousTask();
-	OutData = AsyncTask->GetTask().GetData();
-	return AsyncTask->GetTask().RetrievedEntry();
+	return AsyncTask->GetTask().GetData();
 }
 
 bool FRuntimeAssetCache::ClearCache()
@@ -119,7 +124,7 @@ void FRuntimeAssetCache::AddToAsyncCompletionCounter(int32 Value)
 	check(PendingTasksCounter.GetValue() >= 0);
 }
 
-void FRuntimeAssetCache::WaitAsynchronousCompletion(uint32 Handle)
+void FRuntimeAssetCache::WaitAsynchronousCompletion(int32 Handle)
 {
 	STAT(double ThisTime = 0);
 	{
@@ -135,28 +140,24 @@ void FRuntimeAssetCache::WaitAsynchronousCompletion(uint32 Handle)
 	INC_FLOAT_STAT_BY(STAT_RAC_ASyncWaitTime, (float)ThisTime);
 }
 
-bool FRuntimeAssetCache::GetAsynchronousResults(uint32 Handle, TArray<uint8>& OutData)
+void* FRuntimeAssetCache::GetAsynchronousResults(int32 Handle)
 {
-	FAsyncTask<FRuntimeAssetCacheAsyncWorker>* AsyncTask = NULL;
+	FAsyncTask<FRuntimeAssetCacheAsyncWorker>* AsyncTask = nullptr;
 	{
 		FScopeLock ScopeLock(&SynchronizationObject);
 		PendingTasks.RemoveAndCopyValue(Handle, AsyncTask);
+		AsyncTask->GetTask().FireCompletionDelegate();
 	}
 	check(AsyncTask);
-	if (!AsyncTask->GetTask().RetrievedEntry())
-	{
-		delete AsyncTask;
-		return false;
-	}
-	OutData = AsyncTask->GetTask().GetData();
+
+	void* Data = AsyncTask->GetTask().GetData();
 	delete AsyncTask;
-	check(OutData.Num());
-	return true;
+	return Data;
 }
 
-bool FRuntimeAssetCache::PollAsynchronousCompletion(uint32 Handle)
+bool FRuntimeAssetCache::PollAsynchronousCompletion(int32 Handle)
 {
-	FAsyncTask<FRuntimeAssetCacheAsyncWorker>* AsyncTask = NULL;
+	FAsyncTask<FRuntimeAssetCacheAsyncWorker>* AsyncTask = nullptr;
 	{
 		FScopeLock ScopeLock(&SynchronizationObject);
 		AsyncTask = PendingTasks.FindRef(Handle);
@@ -164,3 +165,15 @@ bool FRuntimeAssetCache::PollAsynchronousCompletion(uint32 Handle)
 	check(AsyncTask);
 	return AsyncTask->IsDone();
 }
+
+void FRuntimeAssetCache::Tick()
+{
+	FScopeLock ScopeLock(&SynchronizationObject);
+	for (auto Kvp : PendingTasks)
+	{
+		if (Kvp.Value->IsDone())
+		{
+			Kvp.Value->GetTask().FireCompletionDelegate();
+		}
+	}
+}
diff --git a/Engine/Source/Runtime/RuntimeAssetCache/Private/RuntimeAssetCache.h b/Engine/Source/Runtime/RuntimeAssetCache/Private/RuntimeAssetCache.h
--- a/Engine/Source/Runtime/RuntimeAssetCache/Private/RuntimeAssetCache.h
+++ b/Engine/Source/Runtime/RuntimeAssetCache/Private/RuntimeAssetCache.h
@@ -3,10 +3,11 @@
 #pragma once
 #include "RuntimeAssetCacheInterface.h"
 
+/** Stats. */
 DECLARE_FLOAT_ACCUMULATOR_STAT_EXTERN(TEXT("RAC ASync Wait Time"), STAT_RAC_ASyncWaitTime, STATGROUP_RAC, );
 
 /** Forward declarations. */
-class FRuntimeAssetCacheBuilderInterface;
+class IRuntimeAssetCacheBuilder;
 class IRuntimeAssetCacheBackend;
 class FRuntimeAssetCacheAsyncWorker;
 class FCacheEntryMetadata;
@@ -19,15 +20,17 @@ class FRuntimeAssetCache : public FRuntimeAssetCacheInterface
 {
 	/** FRuntimeAssetCacheInterface implementation */
 public:
-	virtual bool GetSynchronous(FRuntimeAssetCacheBuilderInterface* CacheBuilder, TArray<uint8>& OutData) override;
-	virtual uint32 GetAsynchronous(FRuntimeAssetCacheBuilderInterface* CacheBuilder) override;
+	virtual void* GetSynchronous(IRuntimeAssetCacheBuilder* CacheBuilder) override;
+	virtual int32 GetAsynchronous(IRuntimeAssetCacheBuilder* CacheBuilder, const FOnRuntimeAssetCacheAsyncComplete& OnCompletionDelegate) override;
+	virtual int32 GetAsynchronous(IRuntimeAssetCacheBuilder* CacheBuilder) override;
 	virtual int32 GetCacheSize(FName Bucket) const override;
 	virtual bool ClearCache() override;
 	virtual bool ClearCache(FName Bucket) override;
-	virtual void WaitAsynchronousCompletion(uint32 Handle) override;
-	virtual bool GetAsynchronousResults(uint32 Handle, TArray<uint8>& OutData) override;
-	virtual bool PollAsynchronousCompletion(uint32 Handle) override;
+	virtual void WaitAsynchronousCompletion(int32 Handle) override;
+	virtual void* GetAsynchronousResults(int32 Handle) override;
+	virtual bool PollAsynchronousCompletion(int32 Handle) override;
 	virtual void AddToAsyncCompletionCounter(int32 Addend) override;
+	virtual void Tick() override;
 	/** End of FRuntimeAssetCacheInterface implementation */
 
 public:
@@ -52,7 +55,7 @@ private:
 	TMap<FName, FRuntimeAssetCacheBucket*> Buckets;
 
 	/** Map of handle to pending task */
-	TMap<uint32, FAsyncTask<FRuntimeAssetCacheAsyncWorker>*> PendingTasks;
+	TMap<int32, FAsyncTask<FRuntimeAssetCacheAsyncWorker>*> PendingTasks;
 
 	/** Pending tasks synchronization object. */
 	FCriticalSection SynchronizationObject;
diff --git a/Engine/Source/Runtime/RuntimeAssetCache/Private/RuntimeAssetCacheAsyncWorker.cpp b/Engine/Source/Runtime/RuntimeAssetCache/Private/RuntimeAssetCacheAsyncWorker.cpp
--- a/Engine/Source/Runtime/RuntimeAssetCache/Private/RuntimeAssetCacheAsyncWorker.cpp
+++ b/Engine/Source/Runtime/RuntimeAssetCache/Private/RuntimeAssetCacheAsyncWorker.cpp
@@ -6,10 +6,10 @@
 #include "RuntimeAssetCacheBucket.h"
 #include "RuntimeAssetCacheBackend.h"
 #include "RuntimeAssetCacheEntryMetadata.h"
-#include "RuntimeAssetCacheInterface.h"
 #include "RuntimeAssetCacheBucketScopeLock.h"
 #include "ScopeExit.h"
 #include "RuntimeAssetCacheModule.h"
+#include "RuntimeAssetCacheBPHooks.h"
 
 /** Stats */
 DEFINE_STAT(STAT_RAC_NumBuilds);
@@ -18,10 +18,15 @@ DEFINE_STAT(STAT_RAC_NumFails);
 DEFINE_STAT(STAT_RAC_NumGets);
 DEFINE_STAT(STAT_RAC_NumPuts);
 
-FRuntimeAssetCacheAsyncWorker::FRuntimeAssetCacheAsyncWorker(FRuntimeAssetCacheBuilderInterface* InCacheBuilder
-	, TMap<FName, FRuntimeAssetCacheBucket*>* InBuckets)
+FRuntimeAssetCacheAsyncWorker::FRuntimeAssetCacheAsyncWorker(IRuntimeAssetCacheBuilder* InCacheBuilder
+	, TMap<FName, FRuntimeAssetCacheBucket*>* InBuckets
+	, int32 InHandle
+	, const FOnRuntimeAssetCacheAsyncComplete& InCompletionCallback)
 	: CacheBuilder(InCacheBuilder)
 	, Buckets(InBuckets)
+	, Handle(InHandle)
+	, CompletionCallback(InCompletionCallback)
+	, bFiredCompletionDelegate(false)
 { }
 
 void FRuntimeAssetCacheAsyncWorker::DoWork()
@@ -31,6 +36,7 @@ void FRuntimeAssetCacheAsyncWorker::DoWork()
 	FString CacheKey = BuildCacheKey(CacheBuilder);
 	FName CacheKeyName = FName(*CacheKey);
 	FRuntimeAssetCacheBucket* CurrentBucket = (*Buckets)[BucketName];
+	int64 CachedDataSize = -1;
 	INC_DWORD_STAT(STAT_RAC_NumGets);
 
 	FCacheEntryMetadata* Metadata = nullptr;
@@ -51,7 +57,7 @@ void FRuntimeAssetCacheAsyncWorker::DoWork()
 			FPlatformProcess::SleepNoStats(0.0f);
 		}
 
-		Metadata = FRuntimeAssetCacheBackend::Get().GetCachedData(BucketName, *CacheKey, Data);
+		Metadata = FRuntimeAssetCacheBackend::Get().GetCachedData(BucketName, *CacheKey, Data, CachedDataSize);
 	}
 
 	ON_SCOPE_EXIT
@@ -76,10 +82,10 @@ void FRuntimeAssetCacheAsyncWorker::DoWork()
 
 	if (Metadata)
 	{
-		check(Data.Num());
+		check(CacheBuilder->GetSerializedDataSize());
 		INC_DWORD_STAT(STAT_RAC_NumCacheHits);
 		FRuntimeAssetCacheBucketScopeLock Guard(*CurrentBucket);
-		Metadata->SetCachedAssetSize(Data.Num());
+		Metadata->SetCachedAssetSize(CacheBuilder->GetSerializedDataSize());
 		Metadata->SetLastAccessTime(FDateTime::Now());
 		bEntryRetrieved = true;
 		return;
@@ -89,35 +95,33 @@ void FRuntimeAssetCacheAsyncWorker::DoWork()
 	{
 		// Failed, cleanup data and return false.
 		INC_DWORD_STAT(STAT_RAC_NumFails);
-		Data.Empty();
+		Data = nullptr;
 		bEntryRetrieved = false;
 		CurrentBucket->RemoveMetadataEntry(CacheKey);
 		return;
 	}
 
-	bool bSuccess;
 	{
 		INC_DWORD_STAT(STAT_RAC_NumBuilds);
 		DECLARE_SCOPE_CYCLE_COUNTER(TEXT("RAC async build time"), STAT_RAC_AsyncBuildTime, STATGROUP_RAC)
-		bSuccess = CacheBuilder->Build(Data);
+		Data = CacheBuilder->Build();
 	}
 
-	if (!bSuccess)
+	if (!Data)
 	{
 		// Failed, cleanup data and return false.
 		INC_DWORD_STAT(STAT_RAC_NumFails);
-		Data.Empty();
 		CurrentBucket->RemoveMetadataEntry(CacheKey);
 		return;
 	}
 
-	checkf(Data.Num(), TEXT("Size of asset to cache cannot be null. Asset cache key: %s"), *CacheKey);
-	checkf(Data.Num() < CurrentBucket->GetSize(), TEXT("Cached asset is bigger than cache size. Increase cache size or reduce asset size. Asset cache key: %s"), *CacheKey);
+	checkf(CacheBuilder->GetSerializedDataSize(), TEXT("Size of asset to cache cannot be null. Asset cache key: %s"), *CacheKey);
+	checkf(CacheBuilder->GetSerializedDataSize() < CurrentBucket->GetSize(), TEXT("Cached asset is bigger than cache size. Increase cache size (%d) or reduce asset size(%d). Asset cache key: %s"), CurrentBucket->GetSize(), CacheBuilder->GetSerializedDataSize(), *CacheKey);
 
 	FRuntimeAssetCacheBucketScopeLock Lock(*CurrentBucket);
 
 	// Do we need to make some space in cache?
-	int32 SizeOfSpaceToFree = CurrentBucket->GetCurrentSize() + Data.Num() - CurrentBucket->GetSize();
+	int32 SizeOfSpaceToFree = CurrentBucket->GetCurrentSize() + CacheBuilder->GetSerializedDataSize() - CurrentBucket->GetSize();
 	if (SizeOfSpaceToFree > 0)
 	{
 		// Remove oldest entries from cache until we can fit upcoming entry.
@@ -133,18 +137,18 @@ void FRuntimeAssetCacheAsyncWorker::DoWork()
 			Metadata->SetLastAccessTime(Now);
 			if (Metadata->GetCachedAssetSize() == 0)
 			{
-				CurrentBucket->AddToCurrentSize(Data.Num());
+				CurrentBucket->AddToCurrentSize(CacheBuilder->GetSerializedDataSize());
 			}
-			Metadata->SetCachedAssetSize(Data.Num());
+			Metadata->SetCachedAssetSize(CacheBuilder->GetSerializedDataSize());
 			Metadata->SetCachedAssetVersion(CacheBuilder->GetAssetVersion());
 			CurrentBucket->AddMetadataEntry(CacheKey, Metadata, false);
 		}
 		else
 		{
-			Metadata = new FCacheEntryMetadata(Now, Data.Num(), CacheBuilder->GetAssetVersion(), CacheKeyName);
+			Metadata = new FCacheEntryMetadata(Now, CacheBuilder->GetSerializedDataSize(), CacheBuilder->GetAssetVersion(), CacheKeyName);
 			CurrentBucket->AddMetadataEntry(CacheKey, Metadata, true);
 		}
-		FRuntimeAssetCacheBackend::Get().PutCachedData(BucketName, *CacheKey, Data, Metadata);
+		FRuntimeAssetCacheBackend::Get().PutCachedData(BucketName, *CacheKey, Data, CacheBuilder->GetSerializedDataSize(), Metadata);
 
 		// Mark that building is finished AFTER putting data into
 		// cache to avoid duplicate builds of the same entry.
@@ -159,6 +163,16 @@ TStatId FRuntimeAssetCacheAsyncWorker::GetStatId()
 	return TStatId();
 }
 
+void FRuntimeAssetCacheAsyncWorker::FireCompletionDelegate()
+{
+	check(IsInGameThread());
+	if (!FiredCompletionDelegate())
+	{
+		bFiredCompletionDelegate = true;
+		CompletionCallback.ExecuteIfBound(Handle, FVoidPtrParam(Data));
+	}
+}
+
 FString FRuntimeAssetCacheAsyncWorker::SanitizeCacheKey(const TCHAR* CacheKey)
 {
 	FString Output;
@@ -207,7 +221,7 @@ FString FRuntimeAssetCacheAsyncWorker::BuildCacheKey(const TCHAR* VersionString,
 	return SanitizeCacheKey(*UnsanitizedCacheKey);
 }
 
-FString FRuntimeAssetCacheAsyncWorker::BuildCacheKey(FRuntimeAssetCacheBuilderInterface* CacheBuilder)
+FString FRuntimeAssetCacheAsyncWorker::BuildCacheKey(IRuntimeAssetCacheBuilder* CacheBuilder)
 {
 	return BuildCacheKey(CacheBuilder->GetBuilderName(), *CacheBuilder->GetAssetUniqueName());
 }
diff --git a/Engine/Source/Runtime/RuntimeAssetCache/Private/RuntimeAssetCacheAsyncWorker.h b/Engine/Source/Runtime/RuntimeAssetCache/Private/RuntimeAssetCacheAsyncWorker.h
--- a/Engine/Source/Runtime/RuntimeAssetCache/Private/RuntimeAssetCacheAsyncWorker.h
+++ b/Engine/Source/Runtime/RuntimeAssetCache/Private/RuntimeAssetCacheAsyncWorker.h
@@ -11,7 +11,7 @@ DECLARE_DWORD_ACCUMULATOR_STAT_EXTERN(TEXT("RAC Num Gets"), STAT_RAC_NumGets, ST
 DECLARE_DWORD_ACCUMULATOR_STAT_EXTERN(TEXT("RAC Num Puts"), STAT_RAC_NumPuts, STATGROUP_RAC, );
 
 /** Forward declarations. */
-class FRuntimeAssetCacheBuilderInterface;
+class IRuntimeAssetCacheBuilder;
 class FRuntimeAssetCacheBucket;
 
 /**
@@ -21,7 +21,7 @@ class FRuntimeAssetCacheAsyncWorker : public FNonAbandonableTask
 {
 public:
 	/** Constructor */
-	FRuntimeAssetCacheAsyncWorker(FRuntimeAssetCacheBuilderInterface* InCacheBuilder, TMap<FName, FRuntimeAssetCacheBucket*>* InBuckets);
+	FRuntimeAssetCacheAsyncWorker(IRuntimeAssetCacheBuilder* InCacheBuilder, TMap<FName, FRuntimeAssetCacheBucket*>* InBuckets, int32 InHandle, const FOnRuntimeAssetCacheAsyncComplete& InCompletionCallback);
 
 	/** Async task interface implementation. */
 	void DoWork();
@@ -29,20 +29,26 @@ public:
 	/** End of async task interface implementation. */
 
 	/** Gets serialized cache data. */
-	TArray<uint8>& GetData()
+	void* GetData()
 	{
 		return Data;
 	}
 
 	/**
-	 * Checks if entry was successfully retrieved from cache.
-	 * @return True on success, false otherwise. Success means either cache hit or miss, as long as we got valid data.
+	 * Fires completion delegate only if it wasn't fired earlier.
 	 */
-	bool RetrievedEntry() const
+	void FireCompletionDelegate();
+
+private:
+	/**
+	* Checks if task already fired completion delegate.
+	* @return true if already fired completion delegate, false otherwise.
+	*/
+	bool FiredCompletionDelegate() const
 	{
-		return bEntryRetrieved;
+		return bFiredCompletionDelegate;
 	}
-private:
+
 	/**
 	* Static function to make sure a cache key contains only legal characters by using an escape.
 	* @param CacheKey Cache key to sanitize
@@ -64,7 +70,7 @@ private:
 	* @param CacheBuilder Builder to create key from.
 	* @return Assembled cache key
 	**/
-	static FString BuildCacheKey(FRuntimeAssetCacheBuilderInterface* CacheBuilder);
+	static FString BuildCacheKey(IRuntimeAssetCacheBuilder* CacheBuilder);
 
 	/**
 	* Removes entries from cache until given number of bytes are freed (or more).
@@ -74,10 +80,10 @@ private:
 	void FreeCacheSpace(FName Bucket, int32 SizeOfSpaceToFreeInBytes);
 
 	/** Cache builder to create cache entry in case of cache miss. */
-	FRuntimeAssetCacheBuilderInterface* CacheBuilder;
+	IRuntimeAssetCacheBuilder* CacheBuilder;
 	
 	/** Data to return to caller. */
-	TArray<uint8> Data;
+	void* Data;
 
 	/** Reference to map of bucket names to their descriptions. */
 	TMap<FName, FRuntimeAssetCacheBucket*>* Buckets;
@@ -87,4 +93,15 @@ private:
 	 * returns false or no CacheBuilder was provided.
 	 */
 	bool bEntryRetrieved;
+
+	/** Completion delegate called when cache entry is retrieved. */
+	FOnRuntimeAssetCacheAsyncComplete CompletionCallback;
+
+	/** Handle uniquely identifying this worker thread. */
+	int32 Handle;
+
+	/** True if completion delegate was already fired, false otherwise. */
+	bool bFiredCompletionDelegate;
+
+	UClass* Class;
 };
\ No newline at end of file
diff --git a/Engine/Source/Runtime/RuntimeAssetCache/Private/RuntimeAssetCacheBPHooks.cpp b/Engine/Source/Runtime/RuntimeAssetCache/Private/RuntimeAssetCacheBPHooks.cpp
new file mode 100644
--- /dev/null
+++ b/Engine/Source/Runtime/RuntimeAssetCache/Private/RuntimeAssetCacheBPHooks.cpp
@@ -0,0 +1,40 @@
+// Copyright 1998-2015 Epic Games, Inc. All Rights Reserved.
+
+#include "RuntimeAssetCachePrivatePCH.h"
+#include "RuntimeAssetCacheModule.h"
+#include "RuntimeAssetCacheBPHooks.h"
+
+FVoidPtrParam URuntimeAssetCacheBPHooks::GetSynchronous(TScriptInterface<IRuntimeAssetCacheBuilder> CacheBuilder)
+{
+	return FVoidPtrParam(GetRuntimeAssetCache().GetSynchronous(static_cast<IRuntimeAssetCacheBuilder*>(CacheBuilder.GetInterface())));
+}
+
+int32 URuntimeAssetCacheBPHooks::GetAsynchronous(TScriptInterface<IRuntimeAssetCacheBuilder> CacheBuilder, const FOnRuntimeAssetCacheAsyncComplete& CompletionDelegate)
+{
+	return GetRuntimeAssetCache().GetAsynchronous(static_cast<IRuntimeAssetCacheBuilder*>(CacheBuilder.GetInterface()), CompletionDelegate);
+}
+
+int32 URuntimeAssetCacheBPHooks::GetCacheSize(FName Bucket)
+{
+	return GetRuntimeAssetCache().GetCacheSize(Bucket);
+}
+
+bool URuntimeAssetCacheBPHooks::ClearCache(FName Bucket)
+{
+	return GetRuntimeAssetCache().ClearCache(Bucket);
+}
+
+void URuntimeAssetCacheBPHooks::WaitAsynchronousCompletion(int32 Handle)
+{
+	GetRuntimeAssetCache().WaitAsynchronousCompletion(Handle);
+}
+
+FVoidPtrParam URuntimeAssetCacheBPHooks::GetAsynchronousResults(int32 Handle)
+{
+	return FVoidPtrParam(GetRuntimeAssetCache().GetAsynchronousResults(Handle));
+}
+
+bool URuntimeAssetCacheBPHooks::PollAsynchronousCompletion(int32 Handle)
+{
+	return GetRuntimeAssetCache().PollAsynchronousCompletion(Handle);
+}
diff --git a/Engine/Source/Runtime/RuntimeAssetCache/Private/RuntimeAssetCacheBackend.cpp b/Engine/Source/Runtime/RuntimeAssetCache/Private/RuntimeAssetCacheBackend.cpp
--- a/Engine/Source/Runtime/RuntimeAssetCache/Private/RuntimeAssetCacheBackend.cpp
+++ b/Engine/Source/Runtime/RuntimeAssetCache/Private/RuntimeAssetCacheBackend.cpp
@@ -6,7 +6,7 @@
 #include "RuntimeAssetCacheEntryMetadata.h"
 
 
-FCacheEntryMetadata* FRuntimeAssetCacheBackend::GetCachedData(const FName Bucket, const TCHAR* CacheKey, TArray<uint8>& OutData)
+FCacheEntryMetadata* FRuntimeAssetCacheBackend::GetCachedData(const FName Bucket, const TCHAR* CacheKey, void*& OutData, int64& OutDataSize)
 {
 	FCacheEntryMetadata* Result = nullptr;
 	FArchive* Ar = CreateReadArchive(Bucket, CacheKey);
@@ -19,17 +19,15 @@ FCacheEntryMetadata* FRuntimeAssetCacheBackend::GetCachedData(const FName Bucket
 
 	int64 TotalSize = Ar->TotalSize();
 	int64 CurrentPosition = Ar->Tell();
-	int64 NumberOfBytesToSerialize = TotalSize - CurrentPosition;
-	OutData.Reset();
-	OutData.AddUninitialized(NumberOfBytesToSerialize);
-
-	Ar->Serialize(OutData.GetData(), NumberOfBytesToSerialize);
+	OutDataSize = TotalSize - CurrentPosition;
+	OutData = new uint8[OutDataSize];
+	Ar->Serialize(OutData, OutDataSize);
 	Ar->Close();
 	delete Ar;
 	return Result;
 }
 
-bool FRuntimeAssetCacheBackend::PutCachedData(const FName Bucket, const TCHAR* CacheKey, TArray<uint8>& InData, FCacheEntryMetadata* Metadata)
+bool FRuntimeAssetCacheBackend::PutCachedData(const FName Bucket, const TCHAR* CacheKey, void* InData, int64 InDataSize, FCacheEntryMetadata* Metadata)
 {
 	bool bResult = false;
 	FArchive* Ar = CreateWriteArchive(Bucket, CacheKey);
@@ -39,7 +37,7 @@ bool FRuntimeAssetCacheBackend::PutCachedData(const FName Bucket, const TCHAR* C
 	}
 
 	*Ar << *Metadata;
-	Ar->Serialize(InData.GetData(), InData.Num());
+	Ar->Serialize(InData, InDataSize);
 	bResult = Ar->Close();
 	delete Ar;
 	return bResult;
diff --git a/Engine/Source/Runtime/RuntimeAssetCache/Private/RuntimeAssetCacheBackend.h b/Engine/Source/Runtime/RuntimeAssetCache/Private/RuntimeAssetCacheBackend.h
--- a/Engine/Source/Runtime/RuntimeAssetCache/Private/RuntimeAssetCacheBackend.h
+++ b/Engine/Source/Runtime/RuntimeAssetCache/Private/RuntimeAssetCacheBackend.h
@@ -29,7 +29,7 @@ public:
 	 * @param OutData Serialized cache entry. Emptied if entry not found.
 	 * @return Metadata descriptor of cache entry on cache hit, nullptr otherwise.
 	 */
-	FCacheEntryMetadata* GetCachedData(FName Bucket, const TCHAR* CacheKey, TArray<uint8>& OutData);
+	FCacheEntryMetadata* GetCachedData(FName Bucket, const TCHAR* CacheKey, void*& OutData, int64& OutDataSize);
 
 	/**
 	* Puts entry to cache.
@@ -39,7 +39,7 @@ public:
 	* @param Metadata descriptor of cache entry.
 	* @return True if successfully put entry into cache.
 	*/
-	bool PutCachedData(FName Bucket, const TCHAR* CacheKey, TArray<uint8>& InData, FCacheEntryMetadata* Metadata);
+	bool PutCachedData(FName Bucket, const TCHAR* CacheKey, void* InData, int64 InDataSize, FCacheEntryMetadata* Metadata);
 
 	/**
 	* Removes entry from cache.
diff --git a/Engine/Source/Runtime/RuntimeAssetCache/Private/RuntimeAssetCacheEntryMetadata.h b/Engine/Source/Runtime/RuntimeAssetCache/Private/RuntimeAssetCacheEntryMetadata.h
--- a/Engine/Source/Runtime/RuntimeAssetCache/Private/RuntimeAssetCacheEntryMetadata.h
+++ b/Engine/Source/Runtime/RuntimeAssetCache/Private/RuntimeAssetCacheEntryMetadata.h
@@ -42,7 +42,7 @@ public:
 		return CachedAssetSize;
 	}
 
-	void SetCachedAssetSize(int32 Value)
+	void SetCachedAssetSize(int64 Value)
 	{
 		CachedAssetSize = Value;
 	}
@@ -76,7 +76,7 @@ public:
 
 private:
 	FDateTime LastAccessTime;
-	int32 CachedAssetSize;
+	int64 CachedAssetSize;
 	int32 CachedAssetVersion;
 	FName Name;
 	FThreadSafeBool bIsBuilding;
diff --git a/Engine/Source/Runtime/RuntimeAssetCache/Private/RuntimeAssetCacheModule.cpp b/Engine/Source/Runtime/RuntimeAssetCache/Private/RuntimeAssetCacheModule.cpp
--- a/Engine/Source/Runtime/RuntimeAssetCache/Private/RuntimeAssetCacheModule.cpp
+++ b/Engine/Source/Runtime/RuntimeAssetCache/Private/RuntimeAssetCacheModule.cpp
@@ -4,6 +4,8 @@
 #include "RuntimeAssetCacheModule.h"
 #include "RuntimeAssetCache.h"
 #include "ModuleManager.h"
+#include "EngineMinimal.h"
+#include "Engine/World.h"
 
 DEFINE_LOG_CATEGORY(RuntimeAssetCache)
 
@@ -30,6 +32,18 @@ public:
 		static FRuntimeAssetCache RuntimeAssetCache = FRuntimeAssetCache();
 		return RuntimeAssetCache;
 	}
+
+	void StartupModule() override
+	{
+		FWorldDelegates::OnWorldTickStart.AddStatic(&FRuntimeAssetCacheModule::TickRuntimeAssetCache);
+	}
+
+	static void TickRuntimeAssetCache(ELevelTick TickType, float DeltaSeconds)
+	{
+		static FRuntimeAssetCacheModuleInterface& Module = FModuleManager::LoadModuleChecked<FRuntimeAssetCacheModuleInterface>("RuntimeAssetCache");
+		static FRuntimeAssetCacheInterface& Interface = Module.GetRuntimeAssetCache();
+		Interface.Tick();
+	}
 };
 
 IMPLEMENT_MODULE(FRuntimeAssetCacheModule, RuntimeAssetCache);
diff --git a/Engine/Source/Runtime/RuntimeAssetCache/Private/RuntimeAssetCachePrivatePCH.h b/Engine/Source/Runtime/RuntimeAssetCache/Private/RuntimeAssetCachePrivatePCH.h
--- a/Engine/Source/Runtime/RuntimeAssetCache/Private/RuntimeAssetCachePrivatePCH.h
+++ b/Engine/Source/Runtime/RuntimeAssetCache/Private/RuntimeAssetCachePrivatePCH.h
@@ -2,7 +2,9 @@
 
 #pragma once
 
-#include "Core.h"
+#include "CoreUObject.h"
+#include "RuntimeAssetCacheBPHooks.h"
+#include "RuntimeAssetCachePublicPCH.h"
 
 /** Stats */
 DECLARE_STATS_GROUP(TEXT("RAC"), STATGROUP_RAC, STATCAT_Advanced);
diff --git a/Engine/Source/Runtime/RuntimeAssetCache/Public/RuntimeAssetCacheBPHooks.h b/Engine/Source/Runtime/RuntimeAssetCache/Public/RuntimeAssetCacheBPHooks.h
new file mode 100644
--- /dev/null
+++ b/Engine/Source/Runtime/RuntimeAssetCache/Public/RuntimeAssetCacheBPHooks.h
@@ -0,0 +1,57 @@
+// Copyright 1998-2015 Epic Games, Inc. All Rights Reserved.
+
+#pragma once
+#include "Kismet/BlueprintFunctionLibrary.h"
+#include "RuntimeAssetCacheInterface.h"
+#include "RuntimeAssetCacheBPHooks.generated.h"
+USTRUCT()
+struct FVoidPtrParam
+{
+	GENERATED_BODY()
+
+	FVoidPtrParam(void* InData)
+		: Data(InData)
+	{ }
+
+	FVoidPtrParam()
+		: Data(nullptr)
+	{ }
+
+
+	void* Data;
+};
+
+/** Delegates. */
+DECLARE_DYNAMIC_DELEGATE_TwoParams(FOnRuntimeAssetCacheAsyncComplete, int32, Handle, FVoidPtrParam, Data);
+
+/** Forward declarations */
+class IRuntimeAssetCacheBuilder;
+class FOnRuntimeAssetCacheAsyncComplete;
+
+UCLASS()
+class URuntimeAssetCacheBPHooks : public UBlueprintFunctionLibrary
+{
+	GENERATED_BODY()
+public:
+
+	UFUNCTION(BlueprintCallable, Category = RuntimeAssetCache)
+	static FVoidPtrParam GetSynchronous(TScriptInterface<IRuntimeAssetCacheBuilder> CacheBuilder);
+
+	UFUNCTION(BlueprintCallable, Category = RuntimeAssetCache)
+	static int32 GetAsynchronous(TScriptInterface<IRuntimeAssetCacheBuilder> CacheBuilder, const FOnRuntimeAssetCacheAsyncComplete& CompletionDelegate);
+
+	UFUNCTION(BlueprintCallable, Category = RuntimeAssetCache)
+	static int32 GetCacheSize(FName Bucket);
+
+	UFUNCTION(BlueprintCallable, Category = RuntimeAssetCache)
+	static bool ClearCache(FName Bucket);
+
+	UFUNCTION(BlueprintCallable, Category = RuntimeAssetCache)
+	static void WaitAsynchronousCompletion(int32 Handle);
+
+	UFUNCTION(BlueprintCallable, Category = RuntimeAssetCache)
+	static FVoidPtrParam GetAsynchronousResults(int32 Handle);
+
+	UFUNCTION(BlueprintCallable, Category = RuntimeAssetCache)
+	static bool PollAsynchronousCompletion(int32 Handle);
+};
\ No newline at end of file
diff --git a/Engine/Source/Runtime/RuntimeAssetCache/Public/RuntimeAssetCacheInterface.h b/Engine/Source/Runtime/RuntimeAssetCache/Public/RuntimeAssetCacheInterface.h
--- a/Engine/Source/Runtime/RuntimeAssetCache/Public/RuntimeAssetCacheInterface.h
+++ b/Engine/Source/Runtime/RuntimeAssetCache/Public/RuntimeAssetCacheInterface.h
@@ -3,7 +3,8 @@
 #pragma once
 
 /** Forward declarations */
-class FRuntimeAssetCacheBuilderInterface;
+class IRuntimeAssetCacheBuilder;
+class FOnRuntimeAssetCacheAsyncComplete;
 
 /**
 * Interface for the Runtime Asset Cache. Cache is split into buckets to cache various assets separately.
@@ -18,17 +19,26 @@ public:
 	/**
 	* Synchronously gets value from cache. If value is not found, builds entry using CacheBuilder and updates cache.
 	* @param CacheBuilder Builder to produce cache key and in the event of a miss.
-	* @param OutData Reference to serialized data.
-	* @return true if the data was retrieved from the cache or the deriver built the data successfully. false can only occur if the plugin returns false.
+	* @return Pointer to retrieved cache entry, nullptr on fail. Fail occurs only when
+	* - there's no entry in cache and CacheBuilder is nullptr or
+	* - CacheBuilder returned nullptr
 	**/
-	virtual bool GetSynchronous(FRuntimeAssetCacheBuilderInterface* CacheBuilder, TArray<uint8>& OutData) = 0;
+	virtual void* GetSynchronous(IRuntimeAssetCacheBuilder* CacheBuilder) = 0;
 
 	/**
 	* Asynchronously checks the cache. If value is not found, builds entry using CacheBuilder and updates cache.
 	* @param CacheBuilder Builder to produce cache key and in the event of a miss, return the data.
+	* @param OnCompletionDelegate Delegate to call when cache is ready. Delegate is called on main thread.
 	* @return Handle to worker.
 	**/
-	virtual uint32 GetAsynchronous(FRuntimeAssetCacheBuilderInterface* CacheBuilder) = 0;
+	virtual int32 GetAsynchronous(IRuntimeAssetCacheBuilder* CacheBuilder, const FOnRuntimeAssetCacheAsyncComplete& OnCompletionDelegate) = 0;
+
+	/**
+	* Asynchronously checks the cache. If value is not found, builds entry using CacheBuilder and updates cache.
+	* @param CacheBuilder Builder to produce cache key and in the event of a miss, return the data.
+	* @return Handle to worker.
+	**/
+	virtual int32 GetAsynchronous(IRuntimeAssetCacheBuilder* CacheBuilder) = 0;
 
 	/**
 	* Gets cache size.
@@ -54,26 +64,32 @@ public:
 	* Waits until worker finishes execution.
 	* @param Handle Worker handle to wait for.
 	*/
-	virtual void WaitAsynchronousCompletion(uint32 Handle) = 0;
+	virtual void WaitAsynchronousCompletion(int32 Handle) = 0;
 
 	/**
 	* Gets asynchronous query results.
 	* @param Handle Worker handle to check.
-	* @param OutData Reference to serialized cache entry.
-	* @return True if successfully retrieved data, false otherwise.
+	* @return Pointer to retrieved cache entry, nullptr on fail. Fail occurs only when
+	* - there's no entry in cache and CacheBuilder is nullptr or
+	* - CacheBuilder returned nullptr
 	*/
-	virtual bool GetAsynchronousResults(uint32 Handle, TArray<uint8>& OutData) = 0;
+	virtual void* GetAsynchronousResults(int32 Handle) = 0;
 
 	/**
 	* Checks if worker finished execution.
 	* @param Handle Worker handle to check.
 	* @return True if execution finished, false otherwise.
 	*/
-	virtual bool PollAsynchronousCompletion(uint32 Handle) = 0;
+	virtual bool PollAsynchronousCompletion(int32 Handle) = 0;
 
 	/**
 	 * Adds a number from the thread safe counter which tracks outstanding async requests. This is used to ensure everything is complete prior to shutdown.
 	 * @param Value to add to counter. Can be negative.
 	 */
 	virtual void AddToAsyncCompletionCounter(int32 Addend) = 0;
+
+	/**
+	 * Ticks async thread.
+	 */
+	virtual void Tick() = 0;
 };
diff --git a/Engine/Source/Runtime/RuntimeAssetCache/Public/RuntimeAssetCachePluginInterface.h b/Engine/Source/Runtime/RuntimeAssetCache/Public/RuntimeAssetCachePluginInterface.h
--- a/Engine/Source/Runtime/RuntimeAssetCache/Public/RuntimeAssetCachePluginInterface.h
+++ b/Engine/Source/Runtime/RuntimeAssetCache/Public/RuntimeAssetCachePluginInterface.h
@@ -1,46 +1,67 @@
 // Copyright 1998-2015 Epic Games, Inc. All Rights Reserved.
 
 #pragma once
+#include "RuntimeAssetCachePluginInterface.generated.h"
+
+UINTERFACE()
+class URuntimeAssetCacheBuilder : public UInterface
+{
+	GENERATED_BODY()
+};
 
 /**
 * Interface for runtime asset cache cache builders.
 * This API will be called concurrently.
 **/
-class FRuntimeAssetCacheBuilderInterface
+class RUNTIMEASSETCACHE_API IRuntimeAssetCacheBuilder
 {
+	GENERATED_BODY()
+
 public:
-	/** Virtual destructor. */
-	virtual ~FRuntimeAssetCacheBuilderInterface()
-	{ }
 
 	/**
 	* Get the builder type name. This is used to categorize cached data to buckets.
 	* @return Type name of builder.
-	**/
+	*/
 	virtual const TCHAR* GetTypeName() const = 0;
 
 	/**
 	* Get the name of cache builder, this is used as part of the cache key.
 	* @return Name of cache builder
-	**/
+	*/
 	virtual const TCHAR* GetBuilderName() const = 0;
 
 	/**
-	* Returns the builder specific part of the cache key. This must be a alphanumeric+underscore
-	* @return Version number of the builder.
-	**/
+	* Returns the name uniquely describing the asset.
+	* @return Asset unique name.
+	*/
 	virtual FString GetAssetUniqueName() const = 0;
 
 	/**
 	* Does the work of creating serialized cache entry.
-	* @param OutData Array of bytes to fill in with the result data
-	* @return true if successful, in the event of failure the cache is not updated and failure is propagated to the original caller.
-	**/
-	virtual bool Build(TArray<uint8>& OutData) = 0;
+	* @return Pointer to contiguous memory block with serialized cache entry on success, nullptr otherwise.
+	*/
+	virtual void* Build() = 0;
 
 	/**
-	 * Gets asset version to rebuild cache if cached asset is too old.
-	 * @return Asset version.
+	 * Get size in bytes of serialized cache entry.
+	 * @return Cache entry size.
 	 */
+	virtual int64 GetSerializedDataSize() = 0;
+
+	/**
+	* Serializes data into archive
+	* @param Ar Archive to serialize to.
+	* @param InData Data to serialize.
+	*/
+	virtual void SerializeData(FArchive& Ar, void* InData)
+	{
+		Ar.Serialize(InData, GetSerializedDataSize());
+	}
+
+	/**
+	* Gets asset version to rebuild cache if cached asset is too old.
+	* @return Asset version.
+	*/
 	virtual int32 GetAssetVersion() = 0;
 };
diff --git a/Engine/Source/Runtime/RuntimeAssetCache/Public/RuntimeAssetCachePublicPCH.h b/Engine/Source/Runtime/RuntimeAssetCache/Public/RuntimeAssetCachePublicPCH.h
new file mode 100644
--- /dev/null
+++ b/Engine/Source/Runtime/RuntimeAssetCache/Public/RuntimeAssetCachePublicPCH.h
@@ -0,0 +1,8 @@
+// Copyright 1998-2015 Epic Games, Inc. All Rights Reserved.
+
+#pragma once
+
+#include "CoreUObject.h"
+#include "RuntimeAssetCacheInterface.h"
+#include "RuntimeAssetCacheModule.h"
+#include "RuntimeAssetCachePluginInterface.h"
diff --git a/Engine/Source/Runtime/RuntimeAssetCache/RuntimeAssetCache.Build.cs b/Engine/Source/Runtime/RuntimeAssetCache/RuntimeAssetCache.Build.cs
--- a/Engine/Source/Runtime/RuntimeAssetCache/RuntimeAssetCache.Build.cs
+++ b/Engine/Source/Runtime/RuntimeAssetCache/RuntimeAssetCache.Build.cs
@@ -7,6 +7,13 @@ public class RuntimeAssetCache : ModuleRules
 {
 	public RuntimeAssetCache(TargetInfo Target)
 	{
-		PrivateDependencyModuleNames.Add("Core");
+		PrivateDependencyModuleNames.AddRange(
+			new string[] {
+				"Core",
+				"CoreUObject",
+				"Engine",
+			});
+
+		SharedPCHHeaderFile = "Runtime/RuntimeAssetCache/Public/RuntimeAssetCachePublicPCH.h";
 	}
 }
