Buildwatcher - Fix shadow variables
Fix double increment in loop

diff --git a/Engine/Source/Runtime/RuntimeAssetCache/Private/RuntimeAssetCache.cpp b/Engine/Source/Runtime/RuntimeAssetCache/Private/RuntimeAssetCache.cpp
--- a/Engine/Source/Runtime/RuntimeAssetCache/Private/RuntimeAssetCache.cpp
+++ b/Engine/Source/Runtime/RuntimeAssetCache/Private/RuntimeAssetCache.cpp
@@ -99,10 +99,10 @@ bool FRuntimeAssetCache::ClearCache()
 	return bResult;
 }
 
-bool FRuntimeAssetCache::ClearCache(FName Bucket)
+bool FRuntimeAssetCache::ClearCache(FName BucketName)
 {
 	/** Tell backend to clean bucket up. */
-	bool bResult = FRuntimeAssetCacheBackend::Get().ClearCache(Bucket);
+	bool bResult = FRuntimeAssetCacheBackend::Get().ClearCache(BucketName);
 
 	if (!bResult)
 	{
diff --git a/Engine/Source/Runtime/RuntimeAssetCache/Private/RuntimeAssetCacheAsyncWorker.cpp b/Engine/Source/Runtime/RuntimeAssetCache/Private/RuntimeAssetCacheAsyncWorker.cpp
--- a/Engine/Source/Runtime/RuntimeAssetCache/Private/RuntimeAssetCacheAsyncWorker.cpp
+++ b/Engine/Source/Runtime/RuntimeAssetCache/Private/RuntimeAssetCacheAsyncWorker.cpp
@@ -131,7 +131,7 @@ void FRuntimeAssetCacheAsyncWorker::DoWork()
 	{
 		DECLARE_SCOPE_CYCLE_COUNTER(TEXT("RAC async put time"), STAT_RAC_PutTime, STATGROUP_RAC)
 		FDateTime Now = FDateTime::Now();
-		FCacheEntryMetadata* Metadata = CurrentBucket->GetMetadata(*CacheKey);
+		Metadata = CurrentBucket->GetMetadata(*CacheKey);
 		if (Metadata)
 		{
 			Metadata->SetLastAccessTime(Now);
