RuntimeAssetCache improvements and fixes.

diff --git a/Engine/Config/BaseEngine.ini b/Engine/Config/BaseEngine.ini
--- a/Engine/Config/BaseEngine.ini
+++ b/Engine/Config/BaseEngine.ini
@@ -1918,6 +1918,7 @@ Foliage=(NumLODs=1,Name=LOCTEXT("FoliageLOD","Foliage"))
 ; +BucketConfigs=(Name="Advertisement", Size=4000000)
 [RuntimeAssetCache]
 +BucketConfigs=(Name="DefaultBucket", Size=10000000)
++PathToRAC=RuntimeAssetCache
 
 ; Derived Data backend graphs
 ; Each of the below backend graph sections contains a set of nodes used to create derived data backed graph.
diff --git a/Engine/Source/Runtime/RuntimeAssetCache/Private/RuntimeAssetCache.cpp b/Engine/Source/Runtime/RuntimeAssetCache/Private/RuntimeAssetCache.cpp
--- a/Engine/Source/Runtime/RuntimeAssetCache/Private/RuntimeAssetCache.cpp
+++ b/Engine/Source/Runtime/RuntimeAssetCache/Private/RuntimeAssetCache.cpp
@@ -57,6 +57,8 @@ int32 FRuntimeAssetCache::GetAsynchronous(IRuntimeAssetCacheBuilder* CacheBuilde
 	/** Make sure task isn't processed twice. */
 	check(!PendingTasks.Contains(Handle));
 
+	checkf(CacheBuilder->IsBuildThreadSafe(), TEXT("CacheBuilder %s Build function is not thread safe, but builder was used in asynchronous code. Use GetSynchronous instead."), CacheBuilder->GetBuilderName());
+
 	FAsyncTask<FRuntimeAssetCacheAsyncWorker>* AsyncTask = new FAsyncTask<FRuntimeAssetCacheAsyncWorker>(CacheBuilder, &Buckets, Handle, OnComplete);
 
 	{
@@ -64,9 +66,12 @@ int32 FRuntimeAssetCache::GetAsynchronous(IRuntimeAssetCacheBuilder* CacheBuilde
 		PendingTasks.Add(Handle, AsyncTask);
 	}
 	AddToAsyncCompletionCounter(1);
+
 	AsyncTask->StartBackgroundTask();
+
 	return Handle;
 }
+
 int32 FRuntimeAssetCache::GetAsynchronous(IRuntimeAssetCacheBuilder* CacheBuilder)
 {
 	return GetAsynchronous(CacheBuilder, FOnRuntimeAssetCacheAsyncComplete());
@@ -74,6 +79,8 @@ int32 FRuntimeAssetCache::GetAsynchronous(IRuntimeAssetCacheBuilder* CacheBuilde
 
 void* FRuntimeAssetCache::GetSynchronous(IRuntimeAssetCacheBuilder* CacheBuilder)
 {
+	checkf(!CacheBuilder->ShouldBuildAsynchronously(), TEXT("CacheBuilder %s can be only called asynchronously."), CacheBuilder->GetBuilderName());
+	
 	FAsyncTask<FRuntimeAssetCacheAsyncWorker>* AsyncTask = new FAsyncTask<FRuntimeAssetCacheAsyncWorker>(CacheBuilder, &Buckets, -1, FOnRuntimeAssetCacheAsyncComplete());
 	AddToAsyncCompletionCounter(1);
 	AsyncTask->StartSynchronousTask();
diff --git a/Engine/Source/Runtime/RuntimeAssetCache/Private/RuntimeAssetCacheAsyncWorker.cpp b/Engine/Source/Runtime/RuntimeAssetCache/Private/RuntimeAssetCacheAsyncWorker.cpp
--- a/Engine/Source/Runtime/RuntimeAssetCache/Private/RuntimeAssetCacheAsyncWorker.cpp
+++ b/Engine/Source/Runtime/RuntimeAssetCache/Private/RuntimeAssetCacheAsyncWorker.cpp
@@ -31,10 +31,22 @@ FRuntimeAssetCacheAsyncWorker::FRuntimeAssetCacheAsyncWorker(IRuntimeAssetCacheB
 
 void FRuntimeAssetCacheAsyncWorker::DoWork()
 {
-	const TCHAR* Bucket = CacheBuilder->GetTypeName();
+	ON_SCOPE_EXIT
+	{
+		/** Make sure completed work counter works properly regardless of where function is exited. */
+		GetRuntimeAssetCache().AddToAsyncCompletionCounter(-1);
+	};
+
+	const TCHAR* Bucket = CacheBuilder->GetBucketConfigName();
 	FName BucketName = FName(Bucket);
 	FString CacheKey = BuildCacheKey(CacheBuilder);
 	FName CacheKeyName = FName(*CacheKey);
+	if (!Buckets->Contains(BucketName))
+	{
+		UE_LOG(RuntimeAssetCache, Warning, TEXT("Caching asset %s to unknown bucket %s. Asset won't be cached."), *CacheBuilder->GetAssetUniqueName(), Bucket);
+		Data = nullptr;
+		return;
+	}
 	FRuntimeAssetCacheBucket* CurrentBucket = (*Buckets)[BucketName];
 	int64 CachedDataSize = -1;
 	INC_DWORD_STAT(STAT_RAC_NumGets);
@@ -60,12 +72,6 @@ void FRuntimeAssetCacheAsyncWorker::DoWork()
 		Metadata = FRuntimeAssetCacheBackend::Get().GetCachedData(BucketName, *CacheKey, Data, CachedDataSize);
 	}
 
-	ON_SCOPE_EXIT
-	{
-		/** Make sure completed work counter works properly regardless of where function is exited. */
-		GetRuntimeAssetCache().AddToAsyncCompletionCounter(-1);
-	};
-
 	/* Entry found. */
 	if (Metadata
 		/* But was saved with older builder version. */
@@ -111,7 +117,7 @@ void FRuntimeAssetCacheAsyncWorker::DoWork()
 	{
 		// Failed, cleanup data and return false.
 		INC_DWORD_STAT(STAT_RAC_NumFails);
-		CurrentBucket->RemoveMetadataEntry(CacheKey);
+		CurrentBucket->RemoveMetadataEntry(CacheKey, true);
 		return;
 	}
 
diff --git a/Engine/Source/Runtime/RuntimeAssetCache/Private/RuntimeAssetCacheBucket.cpp b/Engine/Source/Runtime/RuntimeAssetCache/Private/RuntimeAssetCacheBucket.cpp
--- a/Engine/Source/Runtime/RuntimeAssetCache/Private/RuntimeAssetCacheBucket.cpp
+++ b/Engine/Source/Runtime/RuntimeAssetCache/Private/RuntimeAssetCacheBucket.cpp
@@ -32,12 +32,12 @@ FCacheEntryMetadata* FRuntimeAssetCacheBucket::GetMetadata(const FString& Key)
 	return CacheMetadata.FindOrAdd(Key);
 }
 
-void FRuntimeAssetCacheBucket::RemoveMetadataEntry(const FString& Key)
+void FRuntimeAssetCacheBucket::RemoveMetadataEntry(const FString& Key, bool bBuildFailed)
 {
 	FScopeLock Lock(&MetadataCriticalSection);
 #if UE_BUILD_DEBUG
 	FCacheEntryMetadata* Entry = CacheMetadata.FindOrAdd(Key);
-	if (Entry)
+	if (Entry && !bBuildFailed)
 	{
 		check(!Entry->IsBuilding());
 	}
diff --git a/Engine/Source/Runtime/RuntimeAssetCache/Private/RuntimeAssetCacheBucket.h b/Engine/Source/Runtime/RuntimeAssetCache/Private/RuntimeAssetCacheBucket.h
--- a/Engine/Source/Runtime/RuntimeAssetCache/Private/RuntimeAssetCacheBucket.h
+++ b/Engine/Source/Runtime/RuntimeAssetCache/Private/RuntimeAssetCacheBucket.h
@@ -50,8 +50,9 @@ public:
 	/**
 	* Removes metadata entry for given key.
 	* @param Key Key to remove metadata for.
+	* @param bBuildFailed Indicates that building cache entry failed and some checks should be skipped.
 	*/
-	void RemoveMetadataEntry(const FString& Key);
+	void RemoveMetadataEntry(const FString& Key, bool bBuildFailed = false);
 
 	/**
 	* Adds metadata entry for given key.
diff --git a/Engine/Source/Runtime/RuntimeAssetCache/Private/RuntimeAssetCacheFilesystemBackend.cpp b/Engine/Source/Runtime/RuntimeAssetCache/Private/RuntimeAssetCacheFilesystemBackend.cpp
--- a/Engine/Source/Runtime/RuntimeAssetCache/Private/RuntimeAssetCacheFilesystemBackend.cpp
+++ b/Engine/Source/Runtime/RuntimeAssetCache/Private/RuntimeAssetCacheFilesystemBackend.cpp
@@ -6,24 +6,25 @@
 #include "RuntimeAssetCacheEntryMetadata.h"
 #include "RuntimeAssetCacheBucket.h"
 
-const FString FRuntimeAssetCacheFilesystemBackend::PathToRAC = TEXT("../../RuntimeAssetCache/");
-
 FArchive* FRuntimeAssetCacheFilesystemBackend::CreateReadArchive(FName Bucket, const TCHAR* CacheKey)
 {
-	FString Path = FPaths::Combine(*PathToRAC, *Bucket.ToString(), CacheKey);
-	return IFileManager::Get().CreateFileReader(*Path);
+	return IFileManager::Get().CreateFileReader(CacheKey);
 }
 
 FArchive* FRuntimeAssetCacheFilesystemBackend::CreateWriteArchive(FName Bucket, const TCHAR* CacheKey)
 {
-	FString Path = FPaths::Combine(*PathToRAC, *Bucket.ToString(), CacheKey);
-	return IFileManager::Get().CreateFileWriter(*Path);
+	return IFileManager::Get().CreateFileWriter(CacheKey);
+}
+
+FRuntimeAssetCacheFilesystemBackend::FRuntimeAssetCacheFilesystemBackend()
+{
+	GConfig->GetString(TEXT("RuntimeAssetCache"), TEXT("PathToRAC"), PathToRAC, GEngineIni);
+	PathToRAC = FPaths::GameSavedDir() / PathToRAC;
 }
 
 bool FRuntimeAssetCacheFilesystemBackend::RemoveCacheEntry(const FName Bucket, const TCHAR* CacheKey)
 {
-	FString Path = FPaths::Combine(*PathToRAC, *Bucket.ToString(), CacheKey);
-	return IFileManager::Get().Delete(*Path);
+	return IFileManager::Get().Delete(CacheKey);
 }
 
 bool FRuntimeAssetCacheFilesystemBackend::ClearCache()
@@ -52,6 +53,10 @@ FRuntimeAssetCacheBucket* FRuntimeAssetCacheFilesystemBackend::PreLoadBucket(FNa
 
 		virtual bool Visit(const TCHAR* FilenameOrDirectory, bool bIsDirectory)
 		{
+			if (bIsDirectory)
+			{
+				return true;
+			}
 			FArchive* Ar = Backend->CreateReadArchive(BucketName, FilenameOrDirectory);
 			FCacheEntryMetadata* Metadata = Backend->PreloadMetadata(Ar);
 			Bucket->AddMetadataEntry(*FPaths::GetBaseFilename(FilenameOrDirectory), Metadata, true);
diff --git a/Engine/Source/Runtime/RuntimeAssetCache/Private/RuntimeAssetCacheFilesystemBackend.h b/Engine/Source/Runtime/RuntimeAssetCache/Private/RuntimeAssetCacheFilesystemBackend.h
--- a/Engine/Source/Runtime/RuntimeAssetCache/Private/RuntimeAssetCacheFilesystemBackend.h
+++ b/Engine/Source/Runtime/RuntimeAssetCache/Private/RuntimeAssetCacheFilesystemBackend.h
@@ -10,6 +10,8 @@
 class FRuntimeAssetCacheFilesystemBackend : public FRuntimeAssetCacheBackend
 {
 public:
+	FRuntimeAssetCacheFilesystemBackend();
+
 	/** FRuntimeAssetCacheBackend interface implementation. */
 	virtual bool RemoveCacheEntry(FName Bucket, const TCHAR* CacheKey) override;
 	virtual bool ClearCache() override;
@@ -23,5 +25,5 @@ protected:
 private:
 
 	/** File system path to runtime asset cache. */
-	static const FString PathToRAC;
+	FString PathToRAC;
 };
diff --git a/Engine/Source/Runtime/RuntimeAssetCache/Public/RuntimeAssetCachePluginInterface.h b/Engine/Source/Runtime/RuntimeAssetCache/Public/RuntimeAssetCachePluginInterface.h
--- a/Engine/Source/Runtime/RuntimeAssetCache/Public/RuntimeAssetCachePluginInterface.h
+++ b/Engine/Source/Runtime/RuntimeAssetCache/Public/RuntimeAssetCachePluginInterface.h
@@ -23,7 +23,7 @@ public:
 	* Get the builder type name. This is used to categorize cached data to buckets.
 	* @return Type name of builder.
 	*/
-	virtual const TCHAR* GetTypeName() const = 0;
+	virtual const TCHAR* GetBucketConfigName() const = 0;
 
 	/**
 	* Get the name of cache builder, this is used as part of the cache key.
@@ -64,4 +64,19 @@ public:
 	* @return Asset version.
 	*/
 	virtual int32 GetAssetVersion() = 0;
+
+	/**
+	* Gets whether Build function must be called asynchronously.
+	* @return true if Build should be called asynchronously, false otherwise.
+	*/
+	virtual bool ShouldBuildAsynchronously() const
+	{
+		return false;
+	}
+
+	/**
+	* Gets whether Build function is threadsafe
+	* @return true if Build is threadsafe, false otherwise.
+	*/
+	virtual bool IsBuildThreadSafe() const = 0;
 };
