Allow the RuntimeAssetCache to work with variable or unpredictable file sizes
- FVoidPtrParam now includes DataSize, and this is passed between functions to make it easy to know how big the void* Data is.
- GetSerializedDataSize is no longer needed. Everything now uses the DataSize stored in the MetaData or returned directly from the Build function.

diff --git a/Engine/Source/Runtime/RuntimeAssetCache/Private/RuntimeAssetCache.cpp b/Engine/Source/Runtime/RuntimeAssetCache/Private/RuntimeAssetCache.cpp
--- a/Engine/Source/Runtime/RuntimeAssetCache/Private/RuntimeAssetCache.cpp
+++ b/Engine/Source/Runtime/RuntimeAssetCache/Private/RuntimeAssetCache.cpp
@@ -77,14 +77,14 @@ int32 FRuntimeAssetCache::GetAsynchronous(IRuntimeAssetCacheBuilder* CacheBuilde
 	return GetAsynchronous(CacheBuilder, FOnRuntimeAssetCacheAsyncComplete());
 }
 
-void* FRuntimeAssetCache::GetSynchronous(IRuntimeAssetCacheBuilder* CacheBuilder)
+FVoidPtrParam FRuntimeAssetCache::GetSynchronous(IRuntimeAssetCacheBuilder* CacheBuilder)
 {
 	checkf(!CacheBuilder->ShouldBuildAsynchronously(), TEXT("CacheBuilder %s can be only called asynchronously."), CacheBuilder->GetBuilderName());
 	
 	FAsyncTask<FRuntimeAssetCacheAsyncWorker>* AsyncTask = new FAsyncTask<FRuntimeAssetCacheAsyncWorker>(CacheBuilder, &Buckets, -1, FOnRuntimeAssetCacheAsyncComplete());
 	AddToAsyncCompletionCounter(1);
 	AsyncTask->StartSynchronousTask();
-	return AsyncTask->GetTask().GetData();
+	return AsyncTask->GetTask().GetDataAndSize();
 }
 
 bool FRuntimeAssetCache::ClearCache()
@@ -147,7 +147,7 @@ void FRuntimeAssetCache::WaitAsynchronousCompletion(int32 Handle)
 	INC_FLOAT_STAT_BY(STAT_RAC_ASyncWaitTime, (float)ThisTime);
 }
 
-void* FRuntimeAssetCache::GetAsynchronousResults(int32 Handle)
+FVoidPtrParam FRuntimeAssetCache::GetAsynchronousResults(int32 Handle)
 {
 	FAsyncTask<FRuntimeAssetCacheAsyncWorker>* AsyncTask = nullptr;
 	{
@@ -157,7 +157,7 @@ void* FRuntimeAssetCache::GetAsynchronousResults(int32 Handle)
 	}
 	check(AsyncTask);
 
-	void* Data = AsyncTask->GetTask().GetData();
+	FVoidPtrParam Data = AsyncTask->GetTask().GetDataAndSize();
 	delete AsyncTask;
 	return Data;
 }
diff --git a/Engine/Source/Runtime/RuntimeAssetCache/Private/RuntimeAssetCache.h b/Engine/Source/Runtime/RuntimeAssetCache/Private/RuntimeAssetCache.h
--- a/Engine/Source/Runtime/RuntimeAssetCache/Private/RuntimeAssetCache.h
+++ b/Engine/Source/Runtime/RuntimeAssetCache/Private/RuntimeAssetCache.h
@@ -20,14 +20,14 @@ class FRuntimeAssetCache : public FRuntimeAssetCacheInterface
 {
 	/** FRuntimeAssetCacheInterface implementation */
 public:
-	virtual void* GetSynchronous(IRuntimeAssetCacheBuilder* CacheBuilder) override;
+	virtual FVoidPtrParam GetSynchronous(IRuntimeAssetCacheBuilder* CacheBuilder) override;
 	virtual int32 GetAsynchronous(IRuntimeAssetCacheBuilder* CacheBuilder, const FOnRuntimeAssetCacheAsyncComplete& OnCompletionDelegate) override;
 	virtual int32 GetAsynchronous(IRuntimeAssetCacheBuilder* CacheBuilder) override;
 	virtual int32 GetCacheSize(FName Bucket) const override;
 	virtual bool ClearCache() override;
 	virtual bool ClearCache(FName Bucket) override;
 	virtual void WaitAsynchronousCompletion(int32 Handle) override;
-	virtual void* GetAsynchronousResults(int32 Handle) override;
+	virtual FVoidPtrParam GetAsynchronousResults(int32 Handle) override;
 	virtual bool PollAsynchronousCompletion(int32 Handle) override;
 	virtual void AddToAsyncCompletionCounter(int32 Addend) override;
 	virtual void Tick() override;
diff --git a/Engine/Source/Runtime/RuntimeAssetCache/Private/RuntimeAssetCacheAsyncWorker.cpp b/Engine/Source/Runtime/RuntimeAssetCache/Private/RuntimeAssetCacheAsyncWorker.cpp
--- a/Engine/Source/Runtime/RuntimeAssetCache/Private/RuntimeAssetCacheAsyncWorker.cpp
+++ b/Engine/Source/Runtime/RuntimeAssetCache/Private/RuntimeAssetCacheAsyncWorker.cpp
@@ -41,6 +41,8 @@ void FRuntimeAssetCacheAsyncWorker::DoWork()
 	FName BucketName = FName(Bucket);
 	FString CacheKey = BuildCacheKey(CacheBuilder);
 	FName CacheKeyName = FName(*CacheKey);
+	DataSize = 0;
+	bEntryRetrieved = false;
 	if (!Buckets->Contains(BucketName))
 	{
 		UE_LOG(RuntimeAssetCache, Warning, TEXT("Caching asset %s to unknown bucket %s. Asset won't be cached."), *CacheBuilder->GetAssetUniqueName(), Bucket);
@@ -48,7 +50,6 @@ void FRuntimeAssetCacheAsyncWorker::DoWork()
 		return;
 	}
 	FRuntimeAssetCacheBucket* CurrentBucket = (*Buckets)[BucketName];
-	int64 CachedDataSize = -1;
 	INC_DWORD_STAT(STAT_RAC_NumGets);
 
 	FCacheEntryMetadata* Metadata = nullptr;
@@ -69,7 +70,7 @@ void FRuntimeAssetCacheAsyncWorker::DoWork()
 			FPlatformProcess::SleepNoStats(0.0f);
 		}
 
-		Metadata = FRuntimeAssetCacheBackend::Get().GetCachedData(BucketName, *CacheKey, Data, CachedDataSize);
+		Metadata = FRuntimeAssetCacheBackend::Get().GetCachedData(BucketName, *CacheKey, Data, DataSize);
 	}
 
 	/* Entry found. */
@@ -88,10 +89,8 @@ void FRuntimeAssetCacheAsyncWorker::DoWork()
 
 	if (Metadata)
 	{
-		check(CacheBuilder->GetSerializedDataSize());
 		INC_DWORD_STAT(STAT_RAC_NumCacheHits);
 		FRuntimeAssetCacheBucketScopeLock Guard(*CurrentBucket);
-		Metadata->SetCachedAssetSize(CacheBuilder->GetSerializedDataSize());
 		Metadata->SetLastAccessTime(FDateTime::Now());
 		bEntryRetrieved = true;
 		return;
@@ -102,6 +101,7 @@ void FRuntimeAssetCacheAsyncWorker::DoWork()
 		// Failed, cleanup data and return false.
 		INC_DWORD_STAT(STAT_RAC_NumFails);
 		Data = nullptr;
+		DataSize = 0;
 		bEntryRetrieved = false;
 		CurrentBucket->RemoveMetadataEntry(CacheKey);
 		return;
@@ -110,24 +110,28 @@ void FRuntimeAssetCacheAsyncWorker::DoWork()
 	{
 		INC_DWORD_STAT(STAT_RAC_NumBuilds);
 		DECLARE_SCOPE_CYCLE_COUNTER(TEXT("RAC async build time"), STAT_RAC_AsyncBuildTime, STATGROUP_RAC)
-		Data = CacheBuilder->Build();
+		FVoidPtrParam Results = CacheBuilder->Build();
+		Data = Results.Data;
+		DataSize = Results.DataSize;
 	}
 
 	if (!Data)
 	{
 		// Failed, cleanup data and return false.
 		INC_DWORD_STAT(STAT_RAC_NumFails);
+		DataSize = 0;
+		bEntryRetrieved = false;
 		CurrentBucket->RemoveMetadataEntry(CacheKey, true);
 		return;
 	}
 
-	checkf(CacheBuilder->GetSerializedDataSize(), TEXT("Size of asset to cache cannot be null. Asset cache key: %s"), *CacheKey);
-	checkf(CacheBuilder->GetSerializedDataSize() < CurrentBucket->GetSize(), TEXT("Cached asset is bigger than cache size. Increase cache size (%d) or reduce asset size(%d). Asset cache key: %s"), CurrentBucket->GetSize(), CacheBuilder->GetSerializedDataSize(), *CacheKey);
+	checkf(DataSize, TEXT("Size of asset to cache cannot be null. Asset cache key: %s"), *CacheKey);
+	checkf(DataSize < CurrentBucket->GetSize(), TEXT("Cached asset is bigger than cache size. Increase cache size (%d) or reduce asset size(%d). Asset cache key: %s"), CurrentBucket->GetSize(), DataSize, *CacheKey);
 
 	FRuntimeAssetCacheBucketScopeLock Lock(*CurrentBucket);
 
 	// Do we need to make some space in cache?
-	int32 SizeOfSpaceToFree = CurrentBucket->GetCurrentSize() + CacheBuilder->GetSerializedDataSize() - CurrentBucket->GetSize();
+	int32 SizeOfSpaceToFree = CurrentBucket->GetCurrentSize() + DataSize - CurrentBucket->GetSize();
 	if (SizeOfSpaceToFree > 0)
 	{
 		// Remove oldest entries from cache until we can fit upcoming entry.
@@ -143,18 +147,18 @@ void FRuntimeAssetCacheAsyncWorker::DoWork()
 			Metadata->SetLastAccessTime(Now);
 			if (Metadata->GetCachedAssetSize() == 0)
 			{
-				CurrentBucket->AddToCurrentSize(CacheBuilder->GetSerializedDataSize());
+				CurrentBucket->AddToCurrentSize(DataSize);
 			}
-			Metadata->SetCachedAssetSize(CacheBuilder->GetSerializedDataSize());
+			Metadata->SetCachedAssetSize(DataSize);
 			Metadata->SetCachedAssetVersion(CacheBuilder->GetAssetVersion());
 			CurrentBucket->AddMetadataEntry(CacheKey, Metadata, false);
 		}
 		else
 		{
-			Metadata = new FCacheEntryMetadata(Now, CacheBuilder->GetSerializedDataSize(), CacheBuilder->GetAssetVersion(), CacheKeyName);
+			Metadata = new FCacheEntryMetadata(Now, DataSize, CacheBuilder->GetAssetVersion(), CacheKeyName);
 			CurrentBucket->AddMetadataEntry(CacheKey, Metadata, true);
 		}
-		FRuntimeAssetCacheBackend::Get().PutCachedData(BucketName, *CacheKey, Data, CacheBuilder->GetSerializedDataSize(), Metadata);
+		FRuntimeAssetCacheBackend::Get().PutCachedData(BucketName, *CacheKey, Data, DataSize, Metadata);
 
 		// Mark that building is finished AFTER putting data into
 		// cache to avoid duplicate builds of the same entry.
@@ -175,7 +179,7 @@ void FRuntimeAssetCacheAsyncWorker::FireCompletionDelegate()
 	if (!FiredCompletionDelegate())
 	{
 		bFiredCompletionDelegate = true;
-		CompletionCallback.ExecuteIfBound(Handle, FVoidPtrParam(Data));
+		CompletionCallback.ExecuteIfBound(Handle, FVoidPtrParam(Data, DataSize));
 	}
 }
 
diff --git a/Engine/Source/Runtime/RuntimeAssetCache/Private/RuntimeAssetCacheAsyncWorker.h b/Engine/Source/Runtime/RuntimeAssetCache/Private/RuntimeAssetCacheAsyncWorker.h
--- a/Engine/Source/Runtime/RuntimeAssetCache/Private/RuntimeAssetCacheAsyncWorker.h
+++ b/Engine/Source/Runtime/RuntimeAssetCache/Private/RuntimeAssetCacheAsyncWorker.h
@@ -34,6 +34,18 @@ public:
 		return Data;
 	}
 
+	/** Gets serialized cache data size in bytes. */
+	int64 GetDataSize()
+	{
+		return DataSize;
+	}
+
+	/** Gets serialized cache data and data size. */
+	FVoidPtrParam GetDataAndSize()
+	{
+		return FVoidPtrParam(Data, DataSize);
+	}
+
 	/**
 	 * Fires completion delegate only if it wasn't fired earlier.
 	 */
@@ -85,6 +97,9 @@ private:
 	/** Data to return to caller. */
 	void* Data;
 
+	/** Size of data returned to caller */
+	int64 DataSize;
+
 	/** Reference to map of bucket names to their descriptions. */
 	TMap<FName, FRuntimeAssetCacheBucket*>* Buckets;
 
diff --git a/Engine/Source/Runtime/RuntimeAssetCache/Private/RuntimeAssetCacheBPHooks.cpp b/Engine/Source/Runtime/RuntimeAssetCache/Private/RuntimeAssetCacheBPHooks.cpp
--- a/Engine/Source/Runtime/RuntimeAssetCache/Private/RuntimeAssetCacheBPHooks.cpp
+++ b/Engine/Source/Runtime/RuntimeAssetCache/Private/RuntimeAssetCacheBPHooks.cpp
@@ -6,7 +6,7 @@
 
 FVoidPtrParam URuntimeAssetCacheBPHooks::GetSynchronous(TScriptInterface<IRuntimeAssetCacheBuilder> CacheBuilder)
 {
-	return FVoidPtrParam(GetRuntimeAssetCache().GetSynchronous(static_cast<IRuntimeAssetCacheBuilder*>(CacheBuilder.GetInterface())));
+	return GetRuntimeAssetCache().GetSynchronous(static_cast<IRuntimeAssetCacheBuilder*>(CacheBuilder.GetInterface()));
 }
 
 int32 URuntimeAssetCacheBPHooks::GetAsynchronous(TScriptInterface<IRuntimeAssetCacheBuilder> CacheBuilder, const FOnRuntimeAssetCacheAsyncComplete& CompletionDelegate)
@@ -31,7 +31,7 @@ void URuntimeAssetCacheBPHooks::WaitAsynchronousCompletion(int32 Handle)
 
 FVoidPtrParam URuntimeAssetCacheBPHooks::GetAsynchronousResults(int32 Handle)
 {
-	return FVoidPtrParam(GetRuntimeAssetCache().GetAsynchronousResults(Handle));
+	return GetRuntimeAssetCache().GetAsynchronousResults(Handle);
 }
 
 bool URuntimeAssetCacheBPHooks::PollAsynchronousCompletion(int32 Handle)
diff --git a/Engine/Source/Runtime/RuntimeAssetCache/Public/RuntimeAssetCacheBPHooks.h b/Engine/Source/Runtime/RuntimeAssetCache/Public/RuntimeAssetCacheBPHooks.h
--- a/Engine/Source/Runtime/RuntimeAssetCache/Public/RuntimeAssetCacheBPHooks.h
+++ b/Engine/Source/Runtime/RuntimeAssetCache/Public/RuntimeAssetCacheBPHooks.h
@@ -4,22 +4,7 @@
 #include "Kismet/BlueprintFunctionLibrary.h"
 #include "RuntimeAssetCacheInterface.h"
 #include "RuntimeAssetCacheBPHooks.generated.h"
-USTRUCT()
-struct FVoidPtrParam
-{
-	GENERATED_BODY()
-
-	FVoidPtrParam(void* InData)
-		: Data(InData)
-	{ }
-
-	FVoidPtrParam()
-		: Data(nullptr)
-	{ }
-
 
-	void* Data;
-};
 
 /** Delegates. */
 DECLARE_DYNAMIC_DELEGATE_TwoParams(FOnRuntimeAssetCacheAsyncComplete, int32, Handle, FVoidPtrParam, Data);
diff --git a/Engine/Source/Runtime/RuntimeAssetCache/Public/RuntimeAssetCacheInterface.h b/Engine/Source/Runtime/RuntimeAssetCache/Public/RuntimeAssetCacheInterface.h
--- a/Engine/Source/Runtime/RuntimeAssetCache/Public/RuntimeAssetCacheInterface.h
+++ b/Engine/Source/Runtime/RuntimeAssetCache/Public/RuntimeAssetCacheInterface.h
@@ -1,11 +1,50 @@
 // Copyright 1998-2015 Epic Games, Inc. All Rights Reserved.
 
 #pragma once
+#include "RuntimeAssetCacheInterface.generated.h"
 
 /** Forward declarations */
 class IRuntimeAssetCacheBuilder;
 class FOnRuntimeAssetCacheAsyncComplete;
 
+/**
+* Useful for passing around void* data and size
+*/
+USTRUCT()
+struct FVoidPtrParam
+{
+	GENERATED_BODY()
+
+	FVoidPtrParam(void* InData, int64 InDataSize)
+	: Data(InData)
+	, DataSize(InDataSize)
+	{ }
+
+	FVoidPtrParam()
+		: Data(nullptr)
+		, DataSize(0)
+	{ }
+
+	static FVoidPtrParam NullPtr()
+	{
+		return FVoidPtrParam();
+	}
+
+	FORCEINLINE bool operator!() const
+	{
+		return (!Data || DataSize <= 0);
+	}
+
+	FORCEINLINE operator bool() const
+	{
+		return (Data && DataSize > 0);
+	}
+
+	void* Data;
+	int64 DataSize;
+};
+
+
 /**
 * Interface for the Runtime Asset Cache. Cache is split into buckets to cache various assets separately.
 * Bucket names and maximum sizes are configure via Engine.ini config file using the following syntax
@@ -23,7 +62,7 @@ public:
 	* - there's no entry in cache and CacheBuilder is nullptr or
 	* - CacheBuilder returned nullptr
 	**/
-	virtual void* GetSynchronous(IRuntimeAssetCacheBuilder* CacheBuilder) = 0;
+	virtual FVoidPtrParam GetSynchronous(IRuntimeAssetCacheBuilder* CacheBuilder) = 0;
 
 	/**
 	* Asynchronously checks the cache. If value is not found, builds entry using CacheBuilder and updates cache.
@@ -73,7 +112,7 @@ public:
 	* - there's no entry in cache and CacheBuilder is nullptr or
 	* - CacheBuilder returned nullptr
 	*/
-	virtual void* GetAsynchronousResults(int32 Handle) = 0;
+	virtual FVoidPtrParam GetAsynchronousResults(int32 Handle) = 0;
 
 	/**
 	* Checks if worker finished execution.
diff --git a/Engine/Source/Runtime/RuntimeAssetCache/Public/RuntimeAssetCachePluginInterface.h b/Engine/Source/Runtime/RuntimeAssetCache/Public/RuntimeAssetCachePluginInterface.h
--- a/Engine/Source/Runtime/RuntimeAssetCache/Public/RuntimeAssetCachePluginInterface.h
+++ b/Engine/Source/Runtime/RuntimeAssetCache/Public/RuntimeAssetCachePluginInterface.h
@@ -41,22 +41,16 @@ public:
 	* Does the work of creating serialized cache entry.
 	* @return Pointer to contiguous memory block with serialized cache entry on success, nullptr otherwise.
 	*/
-	virtual void* Build() = 0;
-
-	/**
-	 * Get size in bytes of serialized cache entry.
-	 * @return Cache entry size.
-	 */
-	virtual int64 GetSerializedDataSize() = 0;
+	virtual FVoidPtrParam Build() = 0;
 
 	/**
 	* Serializes data into archive
 	* @param Ar Archive to serialize to.
 	* @param InData Data to serialize.
 	*/
-	virtual void SerializeData(FArchive& Ar, void* InData)
+	virtual void SerializeData(FArchive& Ar, FVoidPtrParam InData)
 	{
-		Ar.Serialize(InData, GetSerializedDataSize());
+		Ar.Serialize(InData.Data, InData.DataSize);
 	}
 
 	/**
