RuntimeAssetCache - added an easy way to create Builders for UObjects. Just derive from the new URuntimeAssetCacheBuilder_ObjectBase class and override a few methods. Look at UExampleTextureCacheBuilder in QAGame for an example.

diff --git a/Engine/Source/Runtime/RuntimeAssetCache/Private/RuntimeAssetCacheBuilders.cpp b/Engine/Source/Runtime/RuntimeAssetCache/Private/RuntimeAssetCacheBuilders.cpp
new file mode 100644
--- /dev/null
+++ b/Engine/Source/Runtime/RuntimeAssetCache/Private/RuntimeAssetCacheBuilders.cpp
@@ -0,0 +1,90 @@
+// Copyright 1998-2015 Epic Games, Inc. All Rights Reserved.
+
+#include "RuntimeAssetCachePrivatePCH.h"
+#include "RuntimeAssetCacheBuilders.h"
+
+void URuntimeAssetCacheBuilder_ObjectBase::SaveNewAssetToCache(UObject* NewAsset)
+{
+	SetAsset(NewAsset);
+	GetFromCacheAsync(OnAssetCacheComplete);
+}
+
+void URuntimeAssetCacheBuilder_ObjectBase::SetAsset(UObject* NewAsset)
+{
+	Asset = NewAsset;
+	OnSetAsset(Asset);
+}
+
+FVoidPtrParam URuntimeAssetCacheBuilder_ObjectBase::Build() override
+{
+	// There was no cached asset, so this is expecting us to return the data that needs to be saved to disk
+	// If we have no asset created yet, just return null. That will trigger the async creation of the asset.
+	// If we do have an asset, serialize it here into a good format and return a pointer to that memory buffer.
+	if (Asset)
+	{
+		int64 DataSize = GetSerializedDataSizeEstimate();
+		void* Result = new uint8[DataSize];
+		FBufferWriter Ar(Result, DataSize, false);
+		Ar.ArIsPersistent = true;
+		SerializeAsset(Ar);
+
+		return FVoidPtrParam(Result, Ar.Tell());
+	}
+
+	// null
+	return FVoidPtrParam::NullPtr();
+}
+
+void URuntimeAssetCacheBuilder_ObjectBase::GetFromCacheAsync(const FOnAssetCacheComplete& OnComplete)
+{
+	OnAssetCacheComplete = OnComplete;
+	GetFromCacheAsyncCompleteDelegate.BindDynamic(this, &URuntimeAssetCacheBuilder_ObjectBase::GetFromCacheAsyncComplete);
+	GetRuntimeAssetCache().GetAsynchronous(this, GetFromCacheAsyncCompleteDelegate);
+}
+
+void URuntimeAssetCacheBuilder_ObjectBase::GetFromCacheAsyncComplete(int32 Handle, FVoidPtrParam DataPtr)
+{
+	if (DataPtr.Data != nullptr)
+	{
+		// Success! Finished loading or saving data from cache
+		// If saving, then we already have the right data and we can just report success
+		if (Asset == nullptr)
+		{
+			// If loading, we now need to serialize the data into a usable format
+
+			// Make sure Asset is set up to be loaded into
+			OnAssetPreLoad();
+
+			FBufferReader Ar(DataPtr.Data, DataPtr.DataSize, false);
+			SerializeAsset(Ar);
+
+			// Perform any specific init functions after load
+			OnAssetPostLoad();
+		}
+
+		// Free the buffer memory on both save and load
+		// On save the buffer gets created in Build()
+		// On load the buffer gets created in FRuntimeAssetCacheBackend::GetCachedData()
+		FMemory::Free(DataPtr.Data);
+
+		// Success!
+		OnAssetCacheComplete.ExecuteIfBound(this, true);
+		CacheHandle = 0;
+	}
+	else
+	{
+		// Data not on disk. Kick off the creation process.
+		// Once complete, call GetFromCacheAsync() again and it will loop back to this function, but should succeed.
+		// But, prevent recursion the second time by checking if the CacheHandle is already set.
+		if (CacheHandle == 0)
+		{
+			CacheHandle = Handle;
+			OnAssetCacheMiss();
+		}
+		else
+		{
+			// Failed
+			OnAssetCacheComplete.ExecuteIfBound(this, false);
+		}
+	}
+}
diff --git a/Engine/Source/Runtime/RuntimeAssetCache/Public/RuntimeAssetCacheBPHooks.h b/Engine/Source/Runtime/RuntimeAssetCache/Public/RuntimeAssetCacheBPHooks.h
--- a/Engine/Source/Runtime/RuntimeAssetCache/Public/RuntimeAssetCacheBPHooks.h
+++ b/Engine/Source/Runtime/RuntimeAssetCache/Public/RuntimeAssetCacheBPHooks.h
@@ -6,9 +6,6 @@
 #include "RuntimeAssetCacheBPHooks.generated.h"
 
 
-/** Delegates. */
-DECLARE_DYNAMIC_DELEGATE_TwoParams(FOnRuntimeAssetCacheAsyncComplete, int32, Handle, FVoidPtrParam, Data);
-
 /** Forward declarations */
 class IRuntimeAssetCacheBuilder;
 class FOnRuntimeAssetCacheAsyncComplete;
diff --git a/Engine/Source/Runtime/RuntimeAssetCache/Public/RuntimeAssetCacheBuilders.h b/Engine/Source/Runtime/RuntimeAssetCache/Public/RuntimeAssetCacheBuilders.h
new file mode 100644
--- /dev/null
+++ b/Engine/Source/Runtime/RuntimeAssetCache/Public/RuntimeAssetCacheBuilders.h
@@ -0,0 +1,146 @@
+// Copyright 1998-2015 Epic Games, Inc. All Rights Reserved.
+
+#pragma once
+
+#include "RuntimeAssetCachePublicPCH.h"
+#include "RuntimeAssetCacheBuilders.generated.h"
+
+DECLARE_DYNAMIC_DELEGATE_TwoParams(FOnAssetCacheComplete, URuntimeAssetCacheBuilder_ObjectBase*, CachedAssetBuilder, bool, Success);
+
+UCLASS(BlueprintType)
+class RUNTIMEASSETCACHE_API URuntimeAssetCacheBuilder_ObjectBase : public UObject, public IRuntimeAssetCacheBuilder
+{
+	GENERATED_BODY()
+public:
+	virtual void PostInitProperties() override
+	{
+		Super::PostInitProperties();
+		CacheHandle = 0;
+	}
+
+	virtual const TCHAR* GetBucketConfigName() const override
+	{
+		return TEXT("DefaultBucket");
+	}
+
+	virtual const TCHAR* GetBuilderName() const override
+	{
+		return TEXT("UObject");
+	}
+
+	virtual FString GetAssetUniqueName() const override
+	{
+		return AssetName;
+	}
+
+	virtual bool IsBuildThreadSafe() const
+	{
+		return true;
+	}
+
+	virtual bool ShouldBuildAsynchronously() const override
+	{
+		return true;
+	}
+
+	virtual int32 GetAssetVersion() override
+	{
+		return AssetVersion;
+	}
+
+	/**
+	 * Override and make a custom serialization function to save/load the important UObject data to disk
+	 */
+	virtual void SerializeAsset(FArchive& Ar)
+	{
+	}
+
+	/**
+	 * Estimate the size (in bytes) of the Data saved in SerializeData. This makes the buffer memory usage is more efficient.
+	 * This is not necessary. However, the closer the estimate is to the actual size the more efficient the memory allocations will be.
+	 */
+	virtual int64 GetSerializedDataSizeEstimate()
+	{
+		return 1024;
+	}
+
+	FVoidPtrParam Build() override;
+
+	/**
+	 * Call this to get the asset named AssetName from the runtime asset cache.
+	 * If the asset does not exist on disk, then OnAssetCacheMiss will be called.
+	 * Implement OnAssetCacheMiss in order to create the asset that you want cached.
+	 */
+	UFUNCTION(BlueprintCallable, Category = RuntimeAssetCache)
+	void GetFromCacheAsync(const FOnAssetCacheComplete& OnComplete);
+
+	UFUNCTION()
+	void GetFromCacheAsyncComplete(int32 Handle, FVoidPtrParam DataPtr);
+
+	/**
+	 * Make sure Asset is set up and ready to be loaded into
+	 */
+	virtual void OnAssetPreLoad() {}
+
+	/**
+	 * Perform any specific init functions after load
+	 */
+	virtual void OnAssetPostLoad() {}
+
+	/**
+	 * When you get OnAssetCacheMiss you need to load/create the asset that is missing.
+	 * Call SaveNewAssetToCache after you're finished creating the asset to save it back into the cache for next time.
+	 * This will then trigger OnAssetCacheComplete like normal, so you don't need additional code to handle it.
+	 */
+	UFUNCTION(BlueprintNativeEvent, Category = RuntimeAssetCache)
+	void OnAssetCacheMiss();
+	virtual void OnAssetCacheMiss_Implementation()
+	{
+		//
+		// Override and create the new asset here (this is where we would render to a render target, then get the result)
+		//
+
+		// Make sure the new asset gets properly cached for next time.
+		SaveNewAssetToCache(Asset);
+	}
+
+	/**
+	 * Call SaveNewAssetToCache to save an asset back into the cache for next time.
+	 * This will then trigger OnAssetCacheComplete like normal, so you don't need additional code to handle it.
+	 */
+	UFUNCTION(BlueprintCallable, Category = RuntimeAssetCache)
+	void SaveNewAssetToCache(UObject* NewAsset);
+
+	/**
+	 * This merely sets Asset and calls OnSetAsset. Use SaveNewAssetToCache to actually cache a new asset.
+	 */
+	void SetAsset(UObject* NewAsset);
+
+	/**
+	 * Please override in your child class to provide easy access to the Asset.
+	 * Declare a specific UObject* variable in your subclass to make it easy for BP nodes to access, e.g. UPROPERTY() UTexture* Texture;
+	 * This function will always get called when the Asset changes
+	 * So set your specific variable here, e.g.  Texture = Cast<UTexture2D>(NewAsset);
+	 */
+	virtual void OnSetAsset(UObject* NewAsset)
+	{
+		// Example:
+		// Texture = Cast<UTexture2D>(NewAsset);
+	}
+
+public:
+	/** The asset version. Changing this will force a new version of the asset to get cached. */
+	UPROPERTY(BlueprintReadWrite, Category = RuntimeAssetCache, meta = (ExposeOnSpawn = "true"))
+	int32 AssetVersion;
+
+	/** The name of the asset. This should be unique per asset, and is used to look it up from the cache. This should be something that can be known without having Asset in memory (so we can look it up in the cache). */
+	UPROPERTY(BlueprintReadWrite, Category = RuntimeAssetCache, meta = (ExposeOnSpawn = "true"))
+	FString AssetName;
+
+private:
+	UObject* Asset;
+	int32 CacheHandle;
+	FOnRuntimeAssetCacheAsyncComplete GetFromCacheAsyncCompleteDelegate;
+	FOnAssetCacheComplete OnAssetCacheComplete;
+};
+
diff --git a/Engine/Source/Runtime/RuntimeAssetCache/Public/RuntimeAssetCacheInterface.h b/Engine/Source/Runtime/RuntimeAssetCache/Public/RuntimeAssetCacheInterface.h
--- a/Engine/Source/Runtime/RuntimeAssetCache/Public/RuntimeAssetCacheInterface.h
+++ b/Engine/Source/Runtime/RuntimeAssetCache/Public/RuntimeAssetCacheInterface.h
@@ -5,7 +5,6 @@
 
 /** Forward declarations */
 class IRuntimeAssetCacheBuilder;
-class FOnRuntimeAssetCacheAsyncComplete;
 
 /**
 * Useful for passing around void* data and size
@@ -45,6 +44,9 @@ struct FVoidPtrParam
 };
 
 
+/** Delegates. */
+DECLARE_DYNAMIC_DELEGATE_TwoParams(FOnRuntimeAssetCacheAsyncComplete, int32, Handle, FVoidPtrParam, Data);
+
 /**
 * Interface for the Runtime Asset Cache. Cache is split into buckets to cache various assets separately.
 * Bucket names and maximum sizes are configure via Engine.ini config file using the following syntax
