Fixed compile error

diff --git a/Engine/Source/Runtime/RuntimeAssetCache/Private/RuntimeAssetCacheBuilders.cpp b/Engine/Source/Runtime/RuntimeAssetCache/Private/RuntimeAssetCacheBuilders.cpp
--- a/Engine/Source/Runtime/RuntimeAssetCache/Private/RuntimeAssetCacheBuilders.cpp
+++ b/Engine/Source/Runtime/RuntimeAssetCache/Private/RuntimeAssetCacheBuilders.cpp
@@ -15,7 +15,7 @@ void URuntimeAssetCacheBuilder_ObjectBase::SetAsset(UObject* NewAsset)
 	OnSetAsset(Asset);
 }
 
-FVoidPtrParam URuntimeAssetCacheBuilder_ObjectBase::Build() override
+FVoidPtrParam URuntimeAssetCacheBuilder_ObjectBase::Build()
 {
 	// There was no cached asset, so this is expecting us to return the data that needs to be saved to disk
 	// If we have no asset created yet, just return null. That will trigger the async creation of the asset.
