Runtime asset cache tweaks.
- Moved example texture builder from QAGame tests to RuntimeAssetCache module for reference/documentation purposes.
- Changed bucket preloading code to use only cache key as file name (without bucket name).

diff --git a/Engine/Source/Runtime/RuntimeAssetCache/Private/RuntimeAssetCacheBuilders.cpp b/Engine/Source/Runtime/RuntimeAssetCache/Private/RuntimeAssetCacheBuilders.cpp
--- a/Engine/Source/Runtime/RuntimeAssetCache/Private/RuntimeAssetCacheBuilders.cpp
+++ b/Engine/Source/Runtime/RuntimeAssetCache/Private/RuntimeAssetCacheBuilders.cpp
@@ -2,6 +2,7 @@
 
 #include "RuntimeAssetCachePrivatePCH.h"
 #include "RuntimeAssetCacheBuilders.h"
+#include "EngineMinimal.h"
 
 void URuntimeAssetCacheBuilder_ObjectBase::SaveNewAssetToCache(UObject* NewAsset)
 {
@@ -88,3 +89,142 @@ void URuntimeAssetCacheBuilder_ObjectBase::GetFromCacheAsyncComplete(int32 Handl
 		}
 	}
 }
+
+void UExampleTextureCacheBuilder::OnSetAsset(UObject* NewAsset)
+{
+	Texture = Cast<UTexture2D>(NewAsset);
+}
+
+void UExampleTextureCacheBuilder::OnAssetCacheMiss_Implementation()
+{
+	// Override and create the new asset here (this is where we would render to a render target, then get the result)
+	// For this example we will simply load an existing texture
+	UTexture2D* NewTexture = LoadObject<UTexture2D>(nullptr, *AssetName);
+
+	// Make sure the new asset gets properly cached for next time.
+	SaveNewAssetToCache(NewTexture);
+}
+
+void UExampleTextureCacheBuilder::SerializeAsset(FArchive& Ar)
+{
+	if (Texture && Texture->PlatformData)
+	{
+		FTexturePlatformData* PlatformData = Texture->PlatformData;
+		UEnum* PixelFormatEnum = UTexture::GetPixelFormatEnum();
+
+		Ar << PlatformData->SizeX;
+		Ar << PlatformData->SizeY;
+		Ar << PlatformData->NumSlices;
+		if (Ar.IsLoading())
+		{
+			FString PixelFormatString;
+			Ar << PixelFormatString;
+			PlatformData->PixelFormat = (EPixelFormat)PixelFormatEnum->FindEnumIndex(*PixelFormatString);
+		}
+		else if (Ar.IsSaving())
+		{
+			FString PixelFormatString = PixelFormatEnum->GetEnum(PlatformData->PixelFormat).GetPlainNameString();
+			Ar << PixelFormatString;
+		}
+
+		int32 NumMips = PlatformData->Mips.Num();
+		int32 FirstMip = 0;
+		int32 LastMip = NumMips;
+
+		TArray<uint32> SavedFlags;
+		if (Ar.IsSaving())
+		{
+			// Force resident mips inline
+			SavedFlags.Empty(NumMips);
+			for (int32 MipIndex = 0; MipIndex < NumMips; ++MipIndex)
+			{
+				SavedFlags.Add(PlatformData->Mips[MipIndex].BulkData.GetBulkDataFlags());
+				PlatformData->Mips[MipIndex].BulkData.SetBulkDataFlags(BULKDATA_ForceInlinePayload | BULKDATA_SingleUse);
+			}
+
+			// Don't save empty Mips
+			while (FirstMip < NumMips && PlatformData->Mips[FirstMip].BulkData.GetBulkDataSize() <= 0)
+			{
+				FirstMip++;
+			}
+			for (int32 MipIndex = FirstMip + 1; MipIndex < NumMips; ++MipIndex)
+			{
+				if (PlatformData->Mips[FirstMip].BulkData.GetBulkDataSize() <= 0)
+				{
+					// This means there are empty tail mips, which should never happen
+					// If it does, simply don't save any mips after this point.
+					LastMip = MipIndex;
+					break;
+				}
+			}
+			int32 NumMipsSaved = LastMip - FirstMip;
+			Ar << NumMipsSaved;
+		}
+
+		if (Ar.IsLoading())
+		{
+			Ar << NumMips;
+			LastMip = NumMips;
+			PlatformData->Mips.Empty(NumMips);
+			for (int32 MipIndex = 0; MipIndex < NumMips; ++MipIndex)
+			{
+				new(PlatformData->Mips) FTexture2DMipMap();
+			}
+		}
+
+		uint32 LockFlags = Ar.IsSaving() ? LOCK_READ_ONLY : LOCK_READ_WRITE;
+		for (int32 MipIndex = FirstMip; MipIndex < LastMip; ++MipIndex)
+		{
+			FTexture2DMipMap& Mip = PlatformData->Mips[MipIndex];
+			Ar << Mip.SizeX;
+			Ar << Mip.SizeY;
+
+			int32 BulkDataSizeInBytes = Mip.BulkData.GetBulkDataSize();
+			Ar << BulkDataSizeInBytes;
+			if (BulkDataSizeInBytes > 0)
+			{
+				void* BulkMipData = Mip.BulkData.Lock(LockFlags);
+				if (Ar.IsLoading())
+				{
+					int32 ElementCount = BulkDataSizeInBytes / Mip.BulkData.GetElementSize();
+					BulkMipData = Mip.BulkData.Realloc(ElementCount);
+				}
+				Ar.Serialize(BulkMipData, BulkDataSizeInBytes);
+				Mip.BulkData.Unlock();
+			}
+		}
+
+		// Restore flags
+		if (Ar.IsSaving())
+		{
+			for (int32 MipIndex = 0; MipIndex < NumMips; ++MipIndex)
+			{
+				PlatformData->Mips[MipIndex].BulkData.SetBulkDataFlags(SavedFlags[MipIndex]);
+			}
+		}
+	}
+}
+
+void UExampleTextureCacheBuilder::OnAssetPreLoad()
+{
+	// Create an object to load the data into
+	UTexture2D* NewTexture = NewObject<UTexture2D>();
+	NewTexture->PlatformData = new FTexturePlatformData();
+	NewTexture->NeverStream = true;
+
+	SetAsset(NewTexture);
+}
+
+void UExampleTextureCacheBuilder::OnAssetPostLoad()
+{
+	Texture->UpdateResource();
+}
+
+int64 UExampleTextureCacheBuilder::GetSerializedDataSizeEstimate()
+{
+	int64 DataSize = sizeof(FTexturePlatformData);
+	DataSize += sizeof(FString) + (sizeof(TCHAR) * 12);						// Guess the size of the pixel format string (most are less than 12 characters, but we don't need to be exact)
+	DataSize += Texture->GetResourceSize(EResourceSizeMode::Exclusive);		// Size of all the mips
+	DataSize += (sizeof(int32) * 3) * Texture->GetNumMips();					// Each mip stores its X and Y size, and its BulkDataSize
+	return DataSize;
+}
diff --git a/Engine/Source/Runtime/RuntimeAssetCache/Private/RuntimeAssetCacheFilesystemBackend.cpp b/Engine/Source/Runtime/RuntimeAssetCache/Private/RuntimeAssetCacheFilesystemBackend.cpp
--- a/Engine/Source/Runtime/RuntimeAssetCache/Private/RuntimeAssetCacheFilesystemBackend.cpp
+++ b/Engine/Source/Runtime/RuntimeAssetCache/Private/RuntimeAssetCacheFilesystemBackend.cpp
@@ -8,14 +8,12 @@
 
 FArchive* FRuntimeAssetCacheFilesystemBackend::CreateReadArchive(FName Bucket, const TCHAR* CacheKey)
 {
-	FString Path = FPaths::Combine(*PathToRAC, *Bucket.ToString(), CacheKey);
-	return IFileManager::Get().CreateFileReader(*Path);
+	return IFileManager::Get().CreateFileReader(CacheKey);
 }
 
 FArchive* FRuntimeAssetCacheFilesystemBackend::CreateWriteArchive(FName Bucket, const TCHAR* CacheKey)
 {
-	FString Path = FPaths::Combine(*PathToRAC, *Bucket.ToString(), CacheKey);
-	return IFileManager::Get().CreateFileWriter(*Path);
+	return IFileManager::Get().CreateFileWriter(CacheKey);
 }
 
 FRuntimeAssetCacheFilesystemBackend::FRuntimeAssetCacheFilesystemBackend()
