CIS fix.

diff --git a/Engine/Source/Runtime/RuntimeAssetCache/Public/RuntimeAssetCacheBuilders.h b/Engine/Source/Runtime/RuntimeAssetCache/Public/RuntimeAssetCacheBuilders.h
--- a/Engine/Source/Runtime/RuntimeAssetCache/Public/RuntimeAssetCacheBuilders.h
+++ b/Engine/Source/Runtime/RuntimeAssetCache/Public/RuntimeAssetCacheBuilders.h
@@ -7,6 +7,8 @@
 
 DECLARE_DYNAMIC_DELEGATE_TwoParams(FOnAssetCacheComplete, URuntimeAssetCacheBuilder_ObjectBase*, CachedAssetBuilder, bool, Success);
 
+class UTexture2D;
+
 UCLASS(BlueprintType)
 class RUNTIMEASSETCACHE_API URuntimeAssetCacheBuilder_ObjectBase : public UObject, public IRuntimeAssetCacheBuilder
 {
@@ -144,3 +146,24 @@ private:
 	FOnAssetCacheComplete OnAssetCacheComplete;
 };
 
+UCLASS(BlueprintType)
+class UExampleTextureCacheBuilder : public URuntimeAssetCacheBuilder_ObjectBase
+{
+	GENERATED_BODY()
+public:
+
+	virtual void OnSetAsset(UObject* NewAsset);
+
+	virtual void OnAssetCacheMiss_Implementation() override;
+
+	virtual void SerializeAsset(FArchive& Ar) override;
+
+	virtual void OnAssetPreLoad() override;
+
+	virtual void OnAssetPostLoad() override;
+
+	virtual int64 GetSerializedDataSizeEstimate() override;
+
+	UPROPERTY(BlueprintReadOnly, Category = "Cache")
+	UTexture2D* Texture;
+};
