OS X nonunity compile fix

diff --git a/Engine/Source/Runtime/RuntimeAssetCache/Private/RuntimeAssetCacheFilesystemBackend.cpp b/Engine/Source/Runtime/RuntimeAssetCache/Private/RuntimeAssetCacheFilesystemBackend.cpp
--- a/Engine/Source/Runtime/RuntimeAssetCache/Private/RuntimeAssetCacheFilesystemBackend.cpp
+++ b/Engine/Source/Runtime/RuntimeAssetCache/Private/RuntimeAssetCacheFilesystemBackend.cpp
@@ -1,6 +1,5 @@
 // Copyright 1998-2015 Epic Games, Inc. All Rights Reserved.
 
-#pragma once
 #include "RuntimeAssetCachePrivatePCH.h"
 #include "RuntimeAssetCacheFilesystemBackend.h"
 #include "RuntimeAssetCacheEntryMetadata.h"
