Mac non unity fix, missing include and removing a SCOPE_CYCLE_COUNTER that is exposing a compiler bug

diff --git a/Engine/Source/Runtime/Apple/MetalRHI/Private/MetalStateCache.cpp b/Engine/Source/Runtime/Apple/MetalRHI/Private/MetalStateCache.cpp
--- a/Engine/Source/Runtime/Apple/MetalRHI/Private/MetalStateCache.cpp
+++ b/Engine/Source/Runtime/Apple/MetalRHI/Private/MetalStateCache.cpp
@@ -578,8 +578,6 @@ void FMetalStateCache::ConditionalUpdateBackBuffer(FMetalSurface& Surface)
 		// update the back buffer texture the first time used this frame
 		if (Surface.Texture == nil)
 		{
-			SCOPE_CYCLE_COUNTER(STAT_MetalMakeDrawableTime);
-
 			uint32 IdleStart = FPlatformTime::Cycles();
 
 			if (CurrentDrawable == nil)
diff --git a/Engine/Source/Runtime/RuntimeAssetCache/Private/RuntimeAssetCache.h b/Engine/Source/Runtime/RuntimeAssetCache/Private/RuntimeAssetCache.h
--- a/Engine/Source/Runtime/RuntimeAssetCache/Private/RuntimeAssetCache.h
+++ b/Engine/Source/Runtime/RuntimeAssetCache/Private/RuntimeAssetCache.h
@@ -2,6 +2,7 @@
 
 #pragma once
 #include "RuntimeAssetCacheInterface.h"
+#include "RuntimeAssetCacheAsyncWorker.h"
 
 /** Stats. */
 DECLARE_FLOAT_ACCUMULATOR_STAT_EXTERN(TEXT("RAC ASync Wait Time"), STAT_RAC_ASyncWaitTime, STATGROUP_RAC, );
@@ -9,7 +10,6 @@ DECLARE_FLOAT_ACCUMULATOR_STAT_EXTERN(TEXT("RAC ASync Wait Time"), STAT_RAC_ASyn
 /** Forward declarations. */
 class IRuntimeAssetCacheBuilder;
 class IRuntimeAssetCacheBackend;
-class FRuntimeAssetCacheAsyncWorker;
 class FCacheEntryMetadata;
 class FRuntimeAssetCacheBucket;
 
