Fixing compilation errors on the Mac and HTML5:

- Removed unnecessary singleton initialization in FRuntimeAssetCacheModule::GetRuntimeAssetCache (failed due to critical section not longer being copyable)
- Added default constructor definition for FHTML5CriticalSection (private copy constructor was added so default constructor would no longer get auto-geterated)

diff --git a/Engine/Source/Runtime/Core/Public/HTML5/HTML5PlatformProcess.h b/Engine/Source/Runtime/Core/Public/HTML5/HTML5PlatformProcess.h
--- a/Engine/Source/Runtime/Core/Public/HTML5/HTML5PlatformProcess.h
+++ b/Engine/Source/Runtime/Core/Public/HTML5/HTML5PlatformProcess.h
@@ -47,8 +47,10 @@ typedef FHTML5PlatformProcess FPlatformProcess;
  */
 class FHTML5CriticalSection
 {
-
 public:
+
+	FHTML5CriticalSection() {}
+
 	/**
 	 * Locks the critical section
 	 */
diff --git a/Engine/Source/Runtime/RuntimeAssetCache/Private/RuntimeAssetCacheModule.cpp b/Engine/Source/Runtime/RuntimeAssetCache/Private/RuntimeAssetCacheModule.cpp
--- a/Engine/Source/Runtime/RuntimeAssetCache/Private/RuntimeAssetCacheModule.cpp
+++ b/Engine/Source/Runtime/RuntimeAssetCache/Private/RuntimeAssetCacheModule.cpp
@@ -29,7 +29,7 @@ public:
 	*/
 	virtual FRuntimeAssetCacheInterface& GetRuntimeAssetCache() override
 	{
-		static FRuntimeAssetCache RuntimeAssetCache = FRuntimeAssetCache();
+		static FRuntimeAssetCache RuntimeAssetCache;
 		return RuntimeAssetCache;
 	}
 
