Copying //UE4/Dev-Framework to //UE4/Dev-Main (Source: //UE4/Dev-Framework @ 3510040)

diff --git a/Engine/Source/Runtime/Engine/Private/Player.cpp b/Engine/Source/Runtime/Engine/Private/Player.cpp
--- a/Engine/Source/Runtime/Engine/Private/Player.cpp
+++ b/Engine/Source/Runtime/Engine/Private/Player.cpp
@@ -32,6 +32,7 @@ FString UPlayer::ConsoleCommand(const FString& Cmd, bool bWriteToLog)
 
 	UConsole* ViewportConsole = (GEngine->GameViewport != nullptr) ? GEngine->GameViewport->ViewportConsole : nullptr;
 	FConsoleOutputDevice StrOut(ViewportConsole);
+	FOutputDevice& EffectiveOutputDevice = (!bWriteToLog || (ViewportConsole != nullptr)) ? StrOut : (FOutputDevice&)(*GLog);
 
 	const int32 CmdLen = Cmd.Len();
 	TCHAR* CommandBuffer = (TCHAR*)FMemory::Malloc((CmdLen + 1)*sizeof(TCHAR));
@@ -47,7 +48,7 @@ FString UPlayer::ConsoleCommand(const FString& Cmd, bool bWriteToLog)
 		// if dissociated with the PC, stop processing commands
 		if (bIsBeacon || PlayerController)
 		{
-			if (!Exec(GetWorld(), Line, StrOut))
+			if (!Exec(GetWorld(), Line, EffectiveOutputDevice))
 			{
 				StrOut.Logf(TEXT("Command not recognized: %s"), Line);
 			}
diff --git a/Engine/Source/Runtime/Engine/Private/StaticMesh.cpp b/Engine/Source/Runtime/Engine/Private/StaticMesh.cpp
--- a/Engine/Source/Runtime/Engine/Private/StaticMesh.cpp
+++ b/Engine/Source/Runtime/Engine/Private/StaticMesh.cpp
@@ -2969,12 +2969,12 @@ void UStaticMesh::CreateNavCollision(const bool bIsUpdate)
 	{
 		UNavCollision* PrevNavCollision = NavCollision;
 
-		if (NavCollision == nullptr || bIsUpdate)
+		if (NavCollision == nullptr)
 		{
 			NavCollision = NewObject<UNavCollision>(this);
 		}
 
-		if (PrevNavCollision)
+		if (PrevNavCollision && PrevNavCollision != NavCollision)
 		{
 			NavCollision->CopyUserSettings(*PrevNavCollision);
 		}
diff --git a/Engine/Source/Runtime/Engine/Private/UnrealEngine.cpp b/Engine/Source/Runtime/Engine/Private/UnrealEngine.cpp
--- a/Engine/Source/Runtime/Engine/Private/UnrealEngine.cpp
+++ b/Engine/Source/Runtime/Engine/Private/UnrealEngine.cpp
@@ -198,8 +198,6 @@
 DEFINE_LOG_CATEGORY(LogEngine);
 IMPLEMENT_MODULE( FEngineModule, Engine );
 
-#define DEBUGGING_VIEWPORT_SIZES 0
-
 #define LOCTEXT_NAMESPACE "UnrealEngine"
 
 void OnChangeEngineCVarRequiringRecreateRenderState(IConsoleVariable* Var)
@@ -739,7 +737,7 @@ public:
 	FWorldContext
 -----------------------------------------------------------------------------*/
 
-void FWorldContext::SetCurrentWorld(UWorld *World)
+void FWorldContext::SetCurrentWorld(UWorld* World)
 {
 	if (World != nullptr)
 	{
@@ -7472,7 +7470,6 @@ void UEngine::LogPerformanceCapture(UWorld* World, const FString& MapName, const
 	if ((World != nullptr) && (World->GetGameViewport() != nullptr))
 	{
 		const FStatUnitData* StatUnitData = World->GetGameViewport()->GetStatUnitData();
-		//	const FStatHitchesData* StatHitchesData = World->GetGameViewport->GetStatHitchesData();
 
 		FAutomationPerformanceSnapshot PerfSnapshot;
 		PerfSnapshot.Changelist = FString::FromInt(ChangeList);
@@ -8363,27 +8360,9 @@ void DrawStatsHUD( UWorld* World, FViewport* Viewport, FCanvas* Canvas, UCanvas*
 		int32 X = (CanvasObject) ? CanvasObject->SizeX - FPSXOffset : Viewport->GetSizeXY().X - FPSXOffset; //??
 		int32 Y = (GEngine->IsStereoscopic3D(Viewport)) ? FMath::TruncToInt(Viewport->GetSizeXY().Y * 0.40f) : FMath::TruncToInt(Viewport->GetSizeXY().Y * 0.20f);
 
-		//give the viewport first shot at drawing stats
+		// give the viewport first shot at drawing stats
 		Y = Viewport->DrawStatsHUD(Canvas, X, Y);
 
-#if DEBUGGING_VIEWPORT_SIZES
-		// Useful for debugging viewport sizing/resizing, especially on external displays
-
-		FCanvasTextItem ViewportTextItem( FVector2D( 0,0 ), FText::GetEmpty(), GEngine->GetSmallFont(), FLinearColor::Blue );
-		ViewportTextItem.EnableShadow( FLinearColor::Black );
-		FString CurrentRes = FString::Printf(TEXT("W = %d, H = %d"), Viewport->GetSizeXY().X,Viewport->GetSizeXY().Y);
-		ViewportTextItem.Text = FText::FromString( CurrentRes );
-		Canvas->DrawItem( ViewportTextItem, 5, Y );
-		ViewportTextItem.Text =  LOCTEXT("00", "00" );
-		Canvas->DrawItem( ViewportTextItem, 5, 5 );
-		ViewportTextItem.Text =  LOCTEXT("0M", "0M" );
-		Canvas->DrawItem( ViewportTextItem, 5,  Viewport->GetSizeXY().Y - 2 );
-		ViewportTextItem.Text =  LOCTEXT("M0", "M0" );
-		Canvas->DrawItem( ViewportTextItem, Viewport->GetSizeXY().X - 25, 5 );
-		ViewportTextItem.Text =  LOCTEXT("MM", "MM" );
-		Canvas->DrawItem( ViewportTextItem, Viewport->GetSizeXY().X - 25, Viewport->GetSizeXY().Y - 25 );
-#endif
-
 		// Render all the simple stats
 		GEngine->RenderEngineStats(World, Viewport, Canvas, StatsXOffset, MessageY, X, Y, &ViewLocation, &ViewRotation);
 
@@ -8690,22 +8669,31 @@ void UEngine::WorldDestroyed( UWorld* InWorld )
 	WorldDestroyedEvent.Broadcast( InWorld );
 }
 
-UWorld* UEngine::GetWorldFromContextObject(const UObject* Object, const bool bChecked) const
+UWorld* UEngine::GetWorldFromContextObject(const UObject* Object, EGetWorldErrorMode ErrorMode) const
 {
-	if (!bChecked && Object == NULL)
+	if (Object == nullptr)
 	{
-		return NULL;
+		switch (ErrorMode)
+		{
+		case EGetWorldErrorMode::Assert:
+			check(Object);
+			break;
+		case EGetWorldErrorMode::LogAndReturnNull:
+			FFrame::KismetExecutionMessage(TEXT("A null object was passed as a world context object to UEngine::GetWorldFromContextObject()."), ELogVerbosity::Error);
+			//UE_LOG(LogEngine, Warning, TEXT("UEngine::GetWorldFromContextObject() passed a nullptr"));
+			break;
+		case EGetWorldErrorMode::ReturnNull:
+			break;
+		}
+		return nullptr;
 	}
 
-	check(Object);
-
-	// @note : GetWorldChecked is not thread safe, so we can't call if called by another thread
 	bool bSupported = true;
-	UWorld* World = ((bChecked && IsInGameThread() )? Object->GetWorldChecked(bSupported) : Object->GetWorld());
+	UWorld* World = (ErrorMode == EGetWorldErrorMode::Assert) ? Object->GetWorldChecked(/*out*/ bSupported) : Object->GetWorld();
 	return (bSupported ? World : GWorld);
 }
 
-TArray<class ULocalPlayer*>::TConstIterator	UEngine::GetLocalPlayerIterator(UWorld *World)
+TArray<class ULocalPlayer*>::TConstIterator	UEngine::GetLocalPlayerIterator(UWorld* World)
 {
 	return GetGamePlayers(World).CreateConstIterator();
 }
diff --git a/Engine/Source/Runtime/Engine/Private/UserDefinedEnum.cpp b/Engine/Source/Runtime/Engine/Private/UserDefinedEnum.cpp
--- a/Engine/Source/Runtime/Engine/Private/UserDefinedEnum.cpp
+++ b/Engine/Source/Runtime/Engine/Private/UserDefinedEnum.cpp
@@ -80,6 +80,9 @@ void UUserDefinedEnum::PostLoad()
 		FEnumEditorUtils::UpgradeDisplayNamesFromMetaData(this);
 	}
 	FEnumEditorUtils::EnsureAllDisplayNamesExist(this);
+
+	// Apply the transactional flag to user defined enums that were not created with it
+	SetFlags(RF_Transactional);
 }
 
 void UUserDefinedEnum::PostEditUndo()
@@ -90,11 +93,10 @@ void UUserDefinedEnum::PostEditUndo()
 
 FString UUserDefinedEnum::GenerateNewEnumeratorName()
 {
-	const TCHAR* EnumNameBase = TEXT("NewEnumerator");
 	FString EnumNameString;
 	do 
 	{
-		EnumNameString = FString::Printf(TEXT("%s%u"), EnumNameBase, UniqueNameIndex);
+		EnumNameString = FString::Printf(TEXT("NewEnumerator%u"), UniqueNameIndex);
 		++UniqueNameIndex;
 	} 
 	while (!FEnumEditorUtils::IsProperNameForUserDefinedEnumerator(this, EnumNameString));
diff --git a/Engine/Source/Runtime/Engine/Private/VisualLogger/VisualLogger.cpp b/Engine/Source/Runtime/Engine/Private/VisualLogger/VisualLogger.cpp
--- a/Engine/Source/Runtime/Engine/Private/VisualLogger/VisualLogger.cpp
+++ b/Engine/Source/Runtime/Engine/Private/VisualLogger/VisualLogger.cpp
@@ -29,7 +29,7 @@ namespace
 {
 	UWorld* GetWorldForVisualLogger(const UObject* Object)
 	{
-		UWorld* World = Object ? GEngine->GetWorldFromContextObject(Object, false) : nullptr;
+		UWorld* World = GEngine->GetWorldFromContextObject(Object, EGetWorldErrorMode::ReturnNull);
 #if WITH_EDITOR
 		UEditorEngine *EEngine = Cast<UEditorEngine>(GEngine);
 		if (GIsEditor && EEngine != nullptr && World == nullptr)
@@ -65,7 +65,7 @@ bool FVisualLogger::CheckVisualLogInputInternal(const UObject* Object, const FLo
 		return false;
 	}
 
-	*World = GEngine->GetWorldFromContextObject(Object, false);
+	*World = GEngine->GetWorldFromContextObject(Object, EGetWorldErrorMode::ReturnNull);
 	if (ensure(*World != nullptr) == false)
 	{
 		return false;
diff --git a/Engine/Source/Runtime/Engine/Public/AI/Navigation/PImplRecastNavMesh.h b/Engine/Source/Runtime/Engine/Public/AI/Navigation/PImplRecastNavMesh.h
--- a/Engine/Source/Runtime/Engine/Public/AI/Navigation/PImplRecastNavMesh.h
+++ b/Engine/Source/Runtime/Engine/Public/AI/Navigation/PImplRecastNavMesh.h
@@ -13,6 +13,7 @@
 #include "AI/Navigation/NavFilters/NavigationQueryFilter.h"
 #include "AI/Navigation/NavigationTypes.h"
 #include "AI/Navigation/RecastNavMesh.h"
+#include "AI/Navigation/RecastQueryFilter.h"
 
 #if WITH_RECAST
 #include "Detour/DetourNavMesh.h"
@@ -26,46 +27,6 @@ class UNavigationSystem;
 
 #define RECAST_VERY_SMALL_AGENT_RADIUS 0.0f
 
-class ENGINE_API FRecastQueryFilter : public INavigationQueryFilterInterface, public dtQueryFilter
-{
-public:
-	FRecastQueryFilter(bool bIsVirtual = true);
-	virtual ~FRecastQueryFilter(){}
-
-	virtual void Reset() override;
-
-	virtual void SetAreaCost(uint8 AreaType, float Cost) override;
-	virtual void SetFixedAreaEnteringCost(uint8 AreaType, float Cost) override;
-	virtual void SetExcludedArea(uint8 AreaType) override;
-	virtual void SetAllAreaCosts(const float* CostArray, const int32 Count) override;
-	virtual void GetAllAreaCosts(float* CostArray, float* FixedCostArray, const int32 Count) const override;
-	virtual void SetBacktrackingEnabled(const bool bBacktracking) override;
-	virtual bool IsBacktrackingEnabled() const override;
-	virtual bool IsEqual(const INavigationQueryFilterInterface* Other) const override;
-	virtual void SetIncludeFlags(uint16 Flags) override;
-	virtual uint16 GetIncludeFlags() const override;
-	virtual void SetExcludeFlags(uint16 Flags) override;
-	virtual uint16 GetExcludeFlags() const override;
-	virtual FVector GetAdjustedEndLocation(const FVector& EndLocation) const override { return EndLocation; }
-	virtual INavigationQueryFilterInterface* CreateCopy() const override;
-
-	const dtQueryFilter* GetAsDetourQueryFilter() const { return this; }
-
-	/** note that it results in losing all area cost setup. Call it before setting anything else */
-	void SetIsVirtual(bool bIsVirtual);
-};
-
-struct ENGINE_API FRecastSpeciaLinkFilter : public dtQuerySpecialLinkFilter
-{
-	FRecastSpeciaLinkFilter(UNavigationSystem* NavSystem, const UObject* Owner) : NavSys(NavSystem), SearchOwner(Owner), CachedOwnerOb(nullptr) {}
-	virtual bool isLinkAllowed(const int32 UserId) const override;
-	virtual void initialize() override;
-
-	UNavigationSystem* NavSys;
-	FWeakObjectPtr SearchOwner;
-	UObject* CachedOwnerOb;
-};
-
 /** Engine Private! - Private Implementation details of ARecastNavMesh */
 class ENGINE_API FPImplRecastNavMesh
 {
diff --git a/Engine/Source/Runtime/Engine/Public/AI/Navigation/RecastQueryFilter.h b/Engine/Source/Runtime/Engine/Public/AI/Navigation/RecastQueryFilter.h
new file mode 100644
--- /dev/null
+++ b/Engine/Source/Runtime/Engine/Public/AI/Navigation/RecastQueryFilter.h
@@ -0,0 +1,57 @@
+// Copyright 1998-2017 Epic Games, Inc. All Rights Reserved.
+
+#pragma once 
+
+#include "CoreMinimal.h"
+#include "AI/Navigation/NavFilters/NavigationQueryFilter.h"
+#include "AI/Navigation/NavigationTypes.h"
+
+#if WITH_RECAST
+#include "Detour/DetourNavMesh.h"
+#include "Detour/DetourNavMeshQuery.h"
+
+#define RECAST_VERY_SMALL_AGENT_RADIUS 0.0f
+
+class UNavigationSystem;
+
+class ENGINE_API FRecastQueryFilter : public INavigationQueryFilterInterface, public dtQueryFilter
+{
+public:
+	FRecastQueryFilter(bool bIsVirtual = true);
+	virtual ~FRecastQueryFilter(){}
+
+	virtual void Reset() override;
+
+	virtual void SetAreaCost(uint8 AreaType, float Cost) override;
+	virtual void SetFixedAreaEnteringCost(uint8 AreaType, float Cost) override;
+	virtual void SetExcludedArea(uint8 AreaType) override;
+	virtual void SetAllAreaCosts(const float* CostArray, const int32 Count) override;
+	virtual void GetAllAreaCosts(float* CostArray, float* FixedCostArray, const int32 Count) const override;
+	virtual void SetBacktrackingEnabled(const bool bBacktracking) override;
+	virtual bool IsBacktrackingEnabled() const override;
+	virtual bool IsEqual(const INavigationQueryFilterInterface* Other) const override;
+	virtual void SetIncludeFlags(uint16 Flags) override;
+	virtual uint16 GetIncludeFlags() const override;
+	virtual void SetExcludeFlags(uint16 Flags) override;
+	virtual uint16 GetExcludeFlags() const override;
+	virtual FVector GetAdjustedEndLocation(const FVector& EndLocation) const override { return EndLocation; }
+	virtual INavigationQueryFilterInterface* CreateCopy() const override;
+
+	const dtQueryFilter* GetAsDetourQueryFilter() const { return this; }
+
+	/** note that it results in losing all area cost setup. Call it before setting anything else */
+	void SetIsVirtual(bool bIsVirtual);
+};
+
+struct ENGINE_API FRecastSpeciaLinkFilter : public dtQuerySpecialLinkFilter
+{
+	FRecastSpeciaLinkFilter(UNavigationSystem* NavSystem, const UObject* Owner) : NavSys(NavSystem), SearchOwner(Owner), CachedOwnerOb(nullptr) {}
+	virtual bool isLinkAllowed(const int32 UserId) const override;
+	virtual void initialize() override;
+
+	UNavigationSystem* NavSys;
+	FWeakObjectPtr SearchOwner;
+	UObject* CachedOwnerOb;
+};
+
+#endif	// WITH_RECAST
diff --git a/Engine/Source/Runtime/Foliage/Private/FoliageStatistics.cpp b/Engine/Source/Runtime/Foliage/Private/FoliageStatistics.cpp
--- a/Engine/Source/Runtime/Foliage/Private/FoliageStatistics.cpp
+++ b/Engine/Source/Runtime/Foliage/Private/FoliageStatistics.cpp
@@ -16,22 +16,23 @@ UFoliageStatistics::UFoliageStatistics(const FObjectInitializer& ObjectInitializ
 
 int32 UFoliageStatistics::FoliageOverlappingSphereCount(UObject* WorldContextObject, const UStaticMesh* StaticMesh, FVector CenterPosition, float Radius)
 {
-	UWorld* World = GEngine->GetWorldFromContextObject(WorldContextObject);
-	const FSphere Sphere(CenterPosition, Radius);
-
 	int32 Count = 0;
-
-	for (TActorIterator<AInstancedFoliageActor> It(World); It; ++It)
+	if (UWorld* World = GEngine->GetWorldFromContextObject(WorldContextObject, EGetWorldErrorMode::LogAndReturnNull))
 	{
-		AInstancedFoliageActor* IFA = *It;
-		if (!IFA->IsPendingKill())
-		{
-			TArray<const UFoliageType*> FoliageTypes;
-			IFA->GetAllFoliageTypesForMesh(StaticMesh, FoliageTypes);
+		const FSphere Sphere(CenterPosition, Radius);
 
-			for (const auto Type : FoliageTypes)
+		for (TActorIterator<AInstancedFoliageActor> It(World); It; ++It)
+		{
+			AInstancedFoliageActor* IFA = *It;
+			if (!IFA->IsPendingKill())
 			{
-				Count += IFA->GetOverlappingSphereCount(Type, Sphere);
+				TArray<const UFoliageType*> FoliageTypes;
+				IFA->GetAllFoliageTypesForMesh(StaticMesh, FoliageTypes);
+
+				for (const auto Type : FoliageTypes)
+				{
+					Count += IFA->GetOverlappingSphereCount(Type, Sphere);
+				}
 			}
 		}
 	}
@@ -41,44 +42,44 @@ int32 UFoliageStatistics::FoliageOverlappingSphereCount(UObject* WorldContextObj
 
 int32 UFoliageStatistics::FoliageOverlappingBoxCount(UObject* WorldContextObject, const UStaticMesh* StaticMesh, FBox Box)
 {
-	UWorld* World = GEngine->GetWorldFromContextObject(WorldContextObject);
-
 	int32 Count = 0;
-
-	for (TActorIterator<AInstancedFoliageActor> It(World); It; ++It)
+	if (UWorld* World = GEngine->GetWorldFromContextObject(WorldContextObject, EGetWorldErrorMode::LogAndReturnNull))
 	{
-		AInstancedFoliageActor* IFA = *It;
-		if (!IFA->IsPendingKill())
+		for (TActorIterator<AInstancedFoliageActor> It(World); It; ++It)
 		{
-			TArray<const UFoliageType*> FoliageTypes;
-			IFA->GetAllFoliageTypesForMesh(StaticMesh, FoliageTypes);
-
-			for (const auto Type : FoliageTypes)
+			AInstancedFoliageActor* IFA = *It;
+			if (!IFA->IsPendingKill())
 			{
-				Count += IFA->GetOverlappingBoxCount(Type, Box);
+				TArray<const UFoliageType*> FoliageTypes;
+				IFA->GetAllFoliageTypesForMesh(StaticMesh, FoliageTypes);
+
+				for (const auto Type : FoliageTypes)
+				{
+					Count += IFA->GetOverlappingBoxCount(Type, Box);
+				}
 			}
 		}
 	}
 
-
 	return Count;
 }
 
 void UFoliageStatistics::FoliageOverlappingBoxTransforms(UObject* WorldContextObject, const UStaticMesh* StaticMesh, FBox Box, TArray<FTransform>& OutTransforms)
 {
-	UWorld* World = GEngine->GetWorldFromContextObject(WorldContextObject);
-
-	for (TActorIterator<AInstancedFoliageActor> It(World); It; ++It)
+	if (UWorld* World = GEngine->GetWorldFromContextObject(WorldContextObject, EGetWorldErrorMode::LogAndReturnNull))
 	{
-		AInstancedFoliageActor* IFA = *It;
-		if (!IFA->IsPendingKill())
+		for (TActorIterator<AInstancedFoliageActor> It(World); It; ++It)
 		{
-			TArray<const UFoliageType*> FoliageTypes;
-			IFA->GetAllFoliageTypesForMesh(StaticMesh, FoliageTypes);
-
-			for (const auto Type : FoliageTypes)
+			AInstancedFoliageActor* IFA = *It;
+			if (!IFA->IsPendingKill())
 			{
-				IFA->GetOverlappingBoxTransforms(Type, Box, OutTransforms);
+				TArray<const UFoliageType*> FoliageTypes;
+				IFA->GetAllFoliageTypesForMesh(StaticMesh, FoliageTypes);
+
+				for (const auto Type : FoliageTypes)
+				{
+					IFA->GetOverlappingBoxTransforms(Type, Box, OutTransforms);
+				}
 			}
 		}
 	}
@@ -86,16 +87,18 @@ void UFoliageStatistics::FoliageOverlappingBoxTransforms(UObject* WorldContextOb
 
 void UFoliageStatistics::FoliageOverlappingMeshCounts_Debug(UObject* WorldContextObject, FVector CenterPosition, float Radius, TMap<UStaticMesh*, int32>& OutMeshCounts)
 {
-	UWorld* const World = GEngine->GetWorldFromContextObject(WorldContextObject);
-	const FSphere Sphere(CenterPosition, Radius);
-
-	OutMeshCounts.Reset();
-	for (TActorIterator<AInstancedFoliageActor> It(World); It; ++It)
+	if (UWorld* World = GEngine->GetWorldFromContextObject(WorldContextObject, EGetWorldErrorMode::LogAndReturnNull))
 	{
-		AInstancedFoliageActor* IFA = *It;
-		if (!IFA->IsPendingKill())
+		const FSphere Sphere(CenterPosition, Radius);
+
+		OutMeshCounts.Reset();
+		for (TActorIterator<AInstancedFoliageActor> It(World); It; ++It)
 		{
-			IFA->GetOverlappingMeshCounts(Sphere, OutMeshCounts);
+			AInstancedFoliageActor* IFA = *It;
+			if (!IFA->IsPendingKill())
+			{
+				IFA->GetOverlappingMeshCounts(Sphere, OutMeshCounts);
+			}
 		}
 	}
 }
diff --git a/Engine/Source/Runtime/GameplayTasks/Private/Tasks/GameplayTask_SpawnActor.cpp b/Engine/Source/Runtime/GameplayTasks/Private/Tasks/GameplayTask_SpawnActor.cpp
--- a/Engine/Source/Runtime/GameplayTasks/Private/Tasks/GameplayTask_SpawnActor.cpp
+++ b/Engine/Source/Runtime/GameplayTasks/Private/Tasks/GameplayTask_SpawnActor.cpp
@@ -41,7 +41,7 @@ UGameplayTask_SpawnActor* UGameplayTask_SpawnActor::SpawnActor(TScriptInterface<
 
 bool UGameplayTask_SpawnActor::BeginSpawningActor(UObject* WorldContextObject, AActor*& SpawnedActor)
 {
-	UWorld* const World = GEngine->GetWorldFromContextObject(WorldContextObject);
+	UWorld* const World = GEngine->GetWorldFromContextObject(WorldContextObject, EGetWorldErrorMode::LogAndReturnNull);
 	if (World)
 	{
 		SpawnedActor = World->SpawnActorDeferred<AActor>(ClassToSpawn, FTransform(CachedSpawnRotation, CachedSpawnLocation), NULL, NULL, ESpawnActorCollisionHandlingMethod::AlwaysSpawn);
diff --git a/Engine/Source/Runtime/LevelSequence/Private/LevelSequencePlayer.cpp b/Engine/Source/Runtime/LevelSequence/Private/LevelSequencePlayer.cpp
--- a/Engine/Source/Runtime/LevelSequence/Private/LevelSequencePlayer.cpp
+++ b/Engine/Source/Runtime/LevelSequence/Private/LevelSequencePlayer.cpp
@@ -42,8 +42,11 @@ ULevelSequencePlayer* ULevelSequencePlayer::CreateLevelSequencePlayer(UObject* W
 		return nullptr;
 	}
 
-	UWorld* World = GEngine->GetWorldFromContextObject(WorldContextObject);
-	check(World != nullptr);
+	UWorld* World = GEngine->GetWorldFromContextObject(WorldContextObject, EGetWorldErrorMode::LogAndReturnNull);
+	if (World == nullptr)
+	{
+		return nullptr;
+	}
 
 	FActorSpawnParameters SpawnParams;
 	SpawnParams.SpawnCollisionHandlingOverride = ESpawnActorCollisionHandlingMethod::AlwaysSpawn;
diff --git a/Engine/Source/Runtime/MovieSceneTracks/Private/Evaluation/MovieSceneMaterialParameterCollectionTemplate.cpp b/Engine/Source/Runtime/MovieSceneTracks/Private/Evaluation/MovieSceneMaterialParameterCollectionTemplate.cpp
--- a/Engine/Source/Runtime/MovieSceneTracks/Private/Evaluation/MovieSceneMaterialParameterCollectionTemplate.cpp
+++ b/Engine/Source/Runtime/MovieSceneTracks/Private/Evaluation/MovieSceneMaterialParameterCollectionTemplate.cpp
@@ -108,7 +108,7 @@ struct FMaterialParameterCollectionExecutionToken : IMovieSceneExecutionToken
 	virtual void Execute(const FMovieSceneContext& Context, const FMovieSceneEvaluationOperand& Operand, FPersistentEvaluationData& PersistentData, IMovieScenePlayer& Player)
 	{
 		UObject* WorldContextObject = Player.GetPlaybackContext();
-		UWorld* World = WorldContextObject ? GEngine->GetWorldFromContextObject(WorldContextObject) : nullptr;
+		UWorld* World = GEngine->GetWorldFromContextObject(WorldContextObject, EGetWorldErrorMode::LogAndReturnNull);
 
 		if (!World)
 		{
diff --git a/Engine/Source/Runtime/NetworkFile/NetworkFile.Build.cs b/Engine/Source/Runtime/NetworkFile/NetworkFile.Build.cs
--- a/Engine/Source/Runtime/NetworkFile/NetworkFile.Build.cs
+++ b/Engine/Source/Runtime/NetworkFile/NetworkFile.Build.cs
@@ -19,7 +19,6 @@ public class NetworkFile : ModuleRules
 				new string[]
 				{ 
 					"DerivedDataCache",
-					"PackageDependencyInfo",
 				}
 				);
 		}
diff --git a/Engine/Source/Runtime/NetworkFile/Private/NetworkPlatformFile.cpp b/Engine/Source/Runtime/NetworkFile/Private/NetworkPlatformFile.cpp
--- a/Engine/Source/Runtime/NetworkFile/Private/NetworkPlatformFile.cpp
+++ b/Engine/Source/Runtime/NetworkFile/Private/NetworkPlatformFile.cpp
@@ -25,12 +25,6 @@
 
 #include "Object.h"
 
-#if WITH_UNREAL_DEVELOPER_TOOLS
-	#include "Developer/PackageDependencyInfo/Public/PackageDependencyInfo.h"
-#endif	//WITH_UNREAL_DEVELOPER_TOOLS
-
-
-
 DEFINE_LOG_CATEGORY(LogNetworkPlatformFile);
 
 FString FNetworkPlatformFile::MP4Extension = TEXT(".mp4");
diff --git a/Engine/Source/Runtime/RuntimeAssetCache/Private/RuntimeAssetCache.cpp b/Engine/Source/Runtime/RuntimeAssetCache/Private/RuntimeAssetCache.cpp
--- a/Engine/Source/Runtime/RuntimeAssetCache/Private/RuntimeAssetCache.cpp
+++ b/Engine/Source/Runtime/RuntimeAssetCache/Private/RuntimeAssetCache.cpp
@@ -42,32 +42,40 @@ FRuntimeAssetCache::~FRuntimeAssetCache()
 
 int32 FRuntimeAssetCache::GetCacheSize(FName Bucket) const
 {
-	return Buckets[Bucket]->GetSize();
+	return ensure(Buckets.Contains(Bucket)) ? Buckets[Bucket]->GetSize() : 0;
 }
 
 int32 FRuntimeAssetCache::GetAsynchronous(IRuntimeAssetCacheBuilder* CacheBuilder, const FOnRuntimeAssetCacheAsyncComplete& OnComplete)
 {
-	int32 Handle = GetNextHandle();
+	if (ensure(CacheBuilder != nullptr))
+	{
+		int32 Handle = GetNextHandle();
 
-	/** Must return a valid handle */
-	check(Handle != 0);
-	
-	/** Make sure task isn't processed twice. */
-	check(!PendingTasks.Contains(Handle));
+		/** Must return a valid handle */
+		check(Handle != 0);
 
-	checkf(CacheBuilder->IsBuildThreadSafe(), TEXT("CacheBuilder %s Build function is not thread safe, but builder was used in asynchronous code. Use GetSynchronous instead."), CacheBuilder->GetBuilderName());
+		/** Make sure task isn't processed twice. */
+		check(!PendingTasks.Contains(Handle));
 
-	FAsyncTask<FRuntimeAssetCacheAsyncWorker>* AsyncTask = new FAsyncTask<FRuntimeAssetCacheAsyncWorker>(CacheBuilder, &Buckets, Handle, OnComplete);
+		checkf(CacheBuilder->IsBuildThreadSafe(), TEXT("CacheBuilder %s Build function is not thread safe, but builder was used in asynchronous code. Use GetSynchronous instead."), CacheBuilder->GetBuilderName());
 
-	{
-		FScopeLock ScopeLock(&SynchronizationObject);
-		PendingTasks.Add(Handle, AsyncTask);
-	}
-	AddToAsyncCompletionCounter(1);
+		FAsyncTask<FRuntimeAssetCacheAsyncWorker>* AsyncTask = new FAsyncTask<FRuntimeAssetCacheAsyncWorker>(CacheBuilder, &Buckets, Handle, OnComplete);
+
+		{
+			FScopeLock ScopeLock(&SynchronizationObject);
+			PendingTasks.Add(Handle, AsyncTask);
+		}
+		AddToAsyncCompletionCounter(1);
 
-	AsyncTask->StartBackgroundTask();
+		AsyncTask->StartBackgroundTask();
 
-	return Handle;
+		return Handle;
+	}
+	else
+	{
+		// Using this with other functions is going to cause knock-on errors
+		return 0;
+	}
 }
 
 int32 FRuntimeAssetCache::GetAsynchronous(IRuntimeAssetCacheBuilder* CacheBuilder)
@@ -77,12 +85,17 @@ int32 FRuntimeAssetCache::GetAsynchronous(IRuntimeAssetCacheBuilder* CacheBuilde
 
 FVoidPtrParam FRuntimeAssetCache::GetSynchronous(IRuntimeAssetCacheBuilder* CacheBuilder)
 {
-	checkf(!CacheBuilder->ShouldBuildAsynchronously(), TEXT("CacheBuilder %s can be only called asynchronously."), CacheBuilder->GetBuilderName());
-	
-	FAsyncTask<FRuntimeAssetCacheAsyncWorker>* AsyncTask = new FAsyncTask<FRuntimeAssetCacheAsyncWorker>(CacheBuilder, &Buckets, -1, FOnRuntimeAssetCacheAsyncComplete());
-	AddToAsyncCompletionCounter(1);
-	AsyncTask->StartSynchronousTask();
-	return AsyncTask->GetTask().GetDataAndSize();
+	if (ensure(CacheBuilder))
+	{
+		if (ensureMsgf(!CacheBuilder->ShouldBuildAsynchronously(), TEXT("CacheBuilder %s can be only called asynchronously."), CacheBuilder->GetBuilderName()))
+		{
+			FAsyncTask<FRuntimeAssetCacheAsyncWorker>* AsyncTask = new FAsyncTask<FRuntimeAssetCacheAsyncWorker>(CacheBuilder, &Buckets, -1, FOnRuntimeAssetCacheAsyncComplete());
+			AddToAsyncCompletionCounter(1);
+			AsyncTask->StartSynchronousTask();
+			return AsyncTask->GetTask().GetDataAndSize();
+		}
+	}
+	return FVoidPtrParam::NullPtr();
 }
 
 bool FRuntimeAssetCache::ClearCache()
@@ -139,8 +152,11 @@ void FRuntimeAssetCache::WaitAsynchronousCompletion(int32 Handle)
 			FScopeLock ScopeLock(&SynchronizationObject);
 			AsyncTask = PendingTasks.FindRef(Handle);
 		}
-		check(AsyncTask);
-		AsyncTask->EnsureCompletion();
+
+		if (ensure(AsyncTask))
+		{
+			AsyncTask->EnsureCompletion();
+		}
 	}
 	INC_FLOAT_STAT_BY(STAT_RAC_ASyncWaitTime, (float)ThisTime);
 }
@@ -151,13 +167,22 @@ FVoidPtrParam FRuntimeAssetCache::GetAsynchronousResults(int32 Handle)
 	{
 		FScopeLock ScopeLock(&SynchronizationObject);
 		PendingTasks.RemoveAndCopyValue(Handle, AsyncTask);
-		AsyncTask->GetTask().FireCompletionDelegate();
+		if (ensure(AsyncTask))
+		{
+			AsyncTask->GetTask().FireCompletionDelegate();
+		}
 	}
-	check(AsyncTask);
 
-	FVoidPtrParam Data = AsyncTask->GetTask().GetDataAndSize();
-	delete AsyncTask;
-	return Data;
+	if (ensure(AsyncTask))
+	{
+		FVoidPtrParam Data = AsyncTask->GetTask().GetDataAndSize();
+		delete AsyncTask;
+		return Data;
+	}
+	else
+	{
+		return FVoidPtrParam::NullPtr();
+	}
 }
 
 bool FRuntimeAssetCache::PollAsynchronousCompletion(int32 Handle)
@@ -167,8 +192,8 @@ bool FRuntimeAssetCache::PollAsynchronousCompletion(int32 Handle)
 		FScopeLock ScopeLock(&SynchronizationObject);
 		AsyncTask = PendingTasks.FindRef(Handle);
 	}
-	check(AsyncTask);
-	return AsyncTask->IsDone();
+	ensure(AsyncTask);
+	return AsyncTask ? AsyncTask->IsDone() : true;
 }
 
 void FRuntimeAssetCache::Tick()
diff --git a/Engine/Source/Runtime/RuntimeAssetCache/Private/RuntimeAssetCacheFilesystemBackend.cpp b/Engine/Source/Runtime/RuntimeAssetCache/Private/RuntimeAssetCacheFilesystemBackend.cpp
--- a/Engine/Source/Runtime/RuntimeAssetCache/Private/RuntimeAssetCacheFilesystemBackend.cpp
+++ b/Engine/Source/Runtime/RuntimeAssetCache/Private/RuntimeAssetCacheFilesystemBackend.cpp
@@ -39,7 +39,12 @@ bool FRuntimeAssetCacheFilesystemBackend::ClearCache()
 
 bool FRuntimeAssetCacheFilesystemBackend::ClearCache(FName Bucket)
 {
-	return IFileManager::Get().DeleteDirectory(*FPaths::Combine(*PathToRAC, *Bucket.ToString()), false, true);
+	if (Bucket != NAME_None)
+	{
+		return IFileManager::Get().DeleteDirectory(*FPaths::Combine(*PathToRAC, *Bucket.ToString()), false, true);
+	}
+
+	return false;
 }
 
 FRuntimeAssetCacheBucket* FRuntimeAssetCacheFilesystemBackend::PreLoadBucket(FName BucketName, int32 BucketSize)
diff --git a/Engine/Source/Runtime/UMG/Private/SlateBlueprintLibrary.cpp b/Engine/Source/Runtime/UMG/Private/SlateBlueprintLibrary.cpp
--- a/Engine/Source/Runtime/UMG/Private/SlateBlueprintLibrary.cpp
+++ b/Engine/Source/Runtime/UMG/Private/SlateBlueprintLibrary.cpp
@@ -60,7 +60,7 @@ void USlateBlueprintLibrary::LocalToViewport(UObject* WorldContextObject, const
 
 void USlateBlueprintLibrary::AbsoluteToViewport(UObject* WorldContextObject, FVector2D AbsoluteDesktopCoordinate, FVector2D& PixelPosition, FVector2D& ViewportPosition)
 {
-	UWorld* World = GEngine->GetWorldFromContextObject(WorldContextObject);
+	UWorld* World = GEngine->GetWorldFromContextObject(WorldContextObject, EGetWorldErrorMode::LogAndReturnNull);
 	if ( World && World->IsGameWorld() )
 	{
 		if ( UGameViewportClient* ViewportClient = World->GetGameViewport() )
@@ -94,7 +94,7 @@ void USlateBlueprintLibrary::ScreenToWidgetLocal(UObject* WorldContextObject, co
 
 void USlateBlueprintLibrary::ScreenToWidgetAbsolute(UObject* WorldContextObject, FVector2D ScreenPosition, FVector2D& AbsoluteCoordinate)
 {
-	UWorld* World = GEngine->GetWorldFromContextObject(WorldContextObject);
+	UWorld* World = GEngine->GetWorldFromContextObject(WorldContextObject, EGetWorldErrorMode::LogAndReturnNull);
 	if ( World && World->IsGameWorld() )
 	{
 		if ( UGameViewportClient* ViewportClient = World->GetGameViewport() )
diff --git a/Engine/Source/Runtime/UMG/Private/WidgetBlueprintLibrary.cpp b/Engine/Source/Runtime/UMG/Private/WidgetBlueprintLibrary.cpp
--- a/Engine/Source/Runtime/UMG/Private/WidgetBlueprintLibrary.cpp
+++ b/Engine/Source/Runtime/UMG/Private/WidgetBlueprintLibrary.cpp
@@ -42,7 +42,7 @@ UUserWidget* UWidgetBlueprintLibrary::Create(UObject* WorldContextObject, TSubcl
 	UUserWidget* UserWidget = nullptr;
 	if ( OwningPlayer == nullptr )
 	{
-		UWorld* World = GEngine->GetWorldFromContextObject(WorldContextObject);
+		UWorld* World = GEngine->GetWorldFromContextObject(WorldContextObject, EGetWorldErrorMode::LogAndReturnNull);
 		UserWidget = CreateWidget<UUserWidget>(World, WidgetType);
 	}
 	else
@@ -502,7 +502,7 @@ void UWidgetBlueprintLibrary::GetAllWidgetsOfClass(UObject* WorldContextObject,
 		return;
 	}
 	 
-	const UWorld* World = GEngine->GetWorldFromContextObject(WorldContextObject);
+	const UWorld* World = GEngine->GetWorldFromContextObject(WorldContextObject, EGetWorldErrorMode::LogAndReturnNull);
 	if ( !World )
 	{
 		return;
@@ -547,7 +547,7 @@ void UWidgetBlueprintLibrary::GetAllWidgetsWithInterface(UObject* WorldContextOb
 		return;
 	}
 
-	const UWorld* World = GEngine->GetWorldFromContextObject(WorldContextObject);
+	const UWorld* World = GEngine->GetWorldFromContextObject(WorldContextObject, EGetWorldErrorMode::LogAndReturnNull);
 	if (!World)
 	{
 		return;
@@ -615,7 +615,7 @@ void UWidgetBlueprintLibrary::GetSafeZonePadding(UObject* WorldContextObject, FV
 
 bool UWidgetBlueprintLibrary::SetHardwareCursor(UObject* WorldContextObject, EMouseCursor::Type CursorShape, FName CursorName, FVector2D HotSpot)
 {
-	UWorld* World = GEngine->GetWorldFromContextObject(WorldContextObject);
+	UWorld* World = GEngine->GetWorldFromContextObject(WorldContextObject, EGetWorldErrorMode::LogAndReturnNull);
 	if ( World && World->IsGameWorld() )
 	{
 		if ( UGameViewportClient* ViewportClient = World->GetGameViewport() )
diff --git a/Engine/Source/Runtime/UMG/Private/WidgetLayoutLibrary.cpp b/Engine/Source/Runtime/UMG/Private/WidgetLayoutLibrary.cpp
--- a/Engine/Source/Runtime/UMG/Private/WidgetLayoutLibrary.cpp
+++ b/Engine/Source/Runtime/UMG/Private/WidgetLayoutLibrary.cpp
@@ -81,7 +81,7 @@ float UWidgetLayoutLibrary::GetViewportScale(UObject* WorldContextObject)
 	{
 		float ViewportScale = 1.0f;
 
-		UWorld* World = GEngine->GetWorldFromContextObject(WorldContextObject);
+		UWorld* World = GEngine->GetWorldFromContextObject(WorldContextObject, EGetWorldErrorMode::LogAndReturnNull);
 		if ( World && World->IsGameWorld() )
 		{
 			if ( UGameViewportClient* ViewportClient = World->GetGameViewport() )
@@ -159,7 +159,7 @@ FVector2D UWidgetLayoutLibrary::GetViewportSize(UObject* WorldContextObject)
 	{
 		FVector2D ViewportSize(1, 1);
 
-		UWorld* World = GEngine->GetWorldFromContextObject(WorldContextObject);
+		UWorld* World = GEngine->GetWorldFromContextObject(WorldContextObject, EGetWorldErrorMode::LogAndReturnNull);
 		if ( World && World->IsGameWorld() )
 		{
 			if ( UGameViewportClient* ViewportClient = World->GetGameViewport() )
@@ -176,7 +176,7 @@ FVector2D UWidgetLayoutLibrary::GetViewportSize(UObject* WorldContextObject)
 
 FGeometry UWidgetLayoutLibrary::GetViewportWidgetGeometry(UObject* WorldContextObject)
 {
-	UWorld* World = GEngine->GetWorldFromContextObject(WorldContextObject);
+	UWorld* World = GEngine->GetWorldFromContextObject(WorldContextObject, EGetWorldErrorMode::LogAndReturnNull);
 	if ( World && World->IsGameWorld() )
 	{
 		if ( UGameViewportClient* ViewportClient = World->GetGameViewport() )
@@ -194,7 +194,7 @@ FGeometry UWidgetLayoutLibrary::GetViewportWidgetGeometry(UObject* WorldContextO
 
 FGeometry UWidgetLayoutLibrary::GetPlayerScreenWidgetGeometry(APlayerController* PlayerController)
 {
-	UWorld* World = GEngine->GetWorldFromContextObject(PlayerController);
+	UWorld* World = GEngine->GetWorldFromContextObject(PlayerController, EGetWorldErrorMode::LogAndReturnNull);
 	if ( World && World->IsGameWorld() )
 	{
 		if ( UGameViewportClient* ViewportClient = World->GetGameViewport() )
@@ -282,7 +282,7 @@ UVerticalBoxSlot* UWidgetLayoutLibrary::SlotAsVerticalBoxSlot(UWidget* Widget)
 
 void UWidgetLayoutLibrary::RemoveAllWidgets(UObject* WorldContextObject)
 {
-	UWorld* World = GEngine->GetWorldFromContextObject(WorldContextObject);
+	UWorld* World = GEngine->GetWorldFromContextObject(WorldContextObject, EGetWorldErrorMode::LogAndReturnNull);
 	if ( World && World->IsGameWorld() )
 	{
 		if ( UGameViewportClient* ViewportClient = World->GetGameViewport() )
