Missing file from CL 10334591

Remove PlatformRequiresOptTextureData and rework how those fields are serialized to prevent crashes when loading cooked data in the editor.
 - An additional bit in the FTexturePlatformData next to the Cubemap bit is used to indicate whether the structure contains non-default OptData. Contents of the OptData field is only serialized when that bit is set.
 - Platforms which do no make use of OptData will automatically not have that data in their cooked byte stream, as the values will remain default. Platform that do set those values will have the data included and the bit set.

diff --git a/Engine/Source/Runtime/RuntimeAssetCache/Private/RuntimeAssetCacheBuilders.cpp b/Engine/Source/Runtime/RuntimeAssetCache/Private/RuntimeAssetCacheBuilders.cpp
--- a/Engine/Source/Runtime/RuntimeAssetCache/Private/RuntimeAssetCacheBuilders.cpp
+++ b/Engine/Source/Runtime/RuntimeAssetCache/Private/RuntimeAssetCacheBuilders.cpp
@@ -125,7 +125,7 @@ void UExampleTextureCacheBuilder::SerializeAsset(FArchive& Ar)
 
 		Ar << PlatformData->SizeX;
 		Ar << PlatformData->SizeY;
-		Ar << PlatformData->NumSlicesCubemapMask;
+		Ar << PlatformData->PackedData;
 		if (Ar.IsLoading())
 		{
 			FString PixelFormatString;
@@ -138,6 +138,11 @@ void UExampleTextureCacheBuilder::SerializeAsset(FArchive& Ar)
 			Ar << PixelFormatString;
 		}
 
+		if (PlatformData->GetHasOptData())
+		{
+			Ar << PlatformData->OptData;
+		}
+
 		int32 NumMips = PlatformData->Mips.Num();
 		int32 FirstMip = 0;
 		int32 LastMip = NumMips;
