Fix RuntimeAssetCacheAsyncWorker GetStatId constness

#rnx

diff --git a/Engine/Source/Runtime/RuntimeAssetCache/Private/RuntimeAssetCacheAsyncWorker.cpp b/Engine/Source/Runtime/RuntimeAssetCache/Private/RuntimeAssetCacheAsyncWorker.cpp
--- a/Engine/Source/Runtime/RuntimeAssetCache/Private/RuntimeAssetCacheAsyncWorker.cpp
+++ b/Engine/Source/Runtime/RuntimeAssetCache/Private/RuntimeAssetCacheAsyncWorker.cpp
@@ -167,7 +167,7 @@ void FRuntimeAssetCacheAsyncWorker::DoWork()
 	bEntryRetrieved = true;
 }
 
-TStatId FRuntimeAssetCacheAsyncWorker::GetStatId()
+TStatId FRuntimeAssetCacheAsyncWorker::GetStatId() const
 {
 	return TStatId();
 }
diff --git a/Engine/Source/Runtime/RuntimeAssetCache/Private/RuntimeAssetCacheAsyncWorker.h b/Engine/Source/Runtime/RuntimeAssetCache/Private/RuntimeAssetCacheAsyncWorker.h
--- a/Engine/Source/Runtime/RuntimeAssetCache/Private/RuntimeAssetCacheAsyncWorker.h
+++ b/Engine/Source/Runtime/RuntimeAssetCache/Private/RuntimeAssetCacheAsyncWorker.h
@@ -36,7 +36,7 @@ public:
 
 	/** Async task interface implementation. */
 	void DoWork();
-	TStatId GetStatId();
+	TStatId GetStatId() const;
 	/** End of async task interface implementation. */
 
 	/** Gets serialized cache data. */
