Fix static analysis error in RuntimeAssetCacheAsyncWorker.cpp by removing unreachable code.

#rnx

diff --git a/Engine/Source/Runtime/RuntimeAssetCache/Private/RuntimeAssetCacheAsyncWorker.cpp b/Engine/Source/Runtime/RuntimeAssetCache/Private/RuntimeAssetCacheAsyncWorker.cpp
--- a/Engine/Source/Runtime/RuntimeAssetCache/Private/RuntimeAssetCacheAsyncWorker.cpp
+++ b/Engine/Source/Runtime/RuntimeAssetCache/Private/RuntimeAssetCacheAsyncWorker.cpp
@@ -95,17 +95,6 @@ void FRuntimeAssetCacheAsyncWorker::DoWork()
 		return;
 	}
 
-	if (!CacheBuilder)
-	{
-		// Failed, cleanup data and return false.
-		INC_DWORD_STAT(STAT_RAC_NumFails);
-		Data = nullptr;
-		DataSize = 0;
-		bEntryRetrieved = false;
-		CurrentBucket->RemoveMetadataEntry(CacheKey);
-		return;
-	}
-
 	{
 		INC_DWORD_STAT(STAT_RAC_NumBuilds);
 		DECLARE_SCOPE_CYCLE_COUNTER(TEXT("RAC async build time"), STAT_RAC_AsyncBuildTime, STATGROUP_RAC)
